- Niagara Insights
- Lifecycle track shows activate / deactivate / complete calls
- Performance track shows GT/RT/GPU/Memory information
- Hovering the events provides more information
- Selecting a time range will update the top 5 list to help locate expensive systems

#Synced-CL 51808504

diff --git a/Engine/Plugins/FX/Niagara/Source/Niagara/Internal/NiagaraTrace.h b/Engine/Plugins/FX/Niagara/Source/Niagara/Internal/NiagaraTrace.h
new file mode 100644
--- /dev/null
+++ b/Engine/Plugins/FX/Niagara/Source/Niagara/Internal/NiagaraTrace.h
@@ -0,0 +1,37 @@
+// Copyright Epic Games, Inc. All Rights Reserved.
+
+#pragma once
+
+#include "CoreMinimal.h"
+#include "ProfilingDebugging/TraceAuxiliary.h"
+#include "Trace/Trace.h"
+
+#if !defined(WITH_NIAGARA_INSIGHTS)
+	#define WITH_NIAGARA_INSIGHTS (UE_TRACE_ENABLED && !IS_PROGRAM && !UE_BUILD_SHIPPING)
+#endif
+
+class UNiagaraComponent;
+
+UE_TRACE_CHANNEL_EXTERN(NiagaraChannel, NIAGARA_API)
+
+#if WITH_NIAGARA_INSIGHTS
+namespace FNiagaraInsights
+{
+	void Update();
+	bool IsChannelActive();
+
+	void NotifyComponentActivate(const UNiagaraComponent* Component, bool bReset, bool bIsScalabilityCull, bool bAwaitingActivationDueToNotReady);
+	void NotifyComponentDeactivate(const UNiagaraComponent* Component, bool bImmediate, bool bIsScalabilityCull);
+	void NotifyComponentComplete(const UNiagaraComponent* Component, bool bExternalCompletion);
+}
+#else
+namespace FNiagaraInsights
+{
+	inline void Update() {}
+	inline bool IsChannelActive() { return false; }
+
+	inline void NotifyComponentActivate(const UNiagaraComponent* Component, bool bReset, bool bIsScalabilityCull, bool bAwaitingActivationDueToNotReady) {}
+	inline void NotifyComponentDeactivate(const UNiagaraComponent* Component, bool bImmediate, bool bIsScalabilityCull) {}
+	inline void NotifyComponentComplete(const UNiagaraComponent* Component, bool bExternalCompletion) {}
+}
+#endif //WITH_NIAGARA_INSIGHTS
diff --git a/Engine/Plugins/FX/Niagara/Source/Niagara/Private/NiagaraComponent.cpp b/Engine/Plugins/FX/Niagara/Source/Niagara/Private/NiagaraComponent.cpp
--- a/Engine/Plugins/FX/Niagara/Source/Niagara/Private/NiagaraComponent.cpp
+++ b/Engine/Plugins/FX/Niagara/Source/Niagara/Private/NiagaraComponent.cpp
@@ -20,6 +20,7 @@
 #include "NiagaraStats.h"
 #include "NiagaraSystem.h"
 #include "NiagaraSystemInstanceController.h"
+#include "NiagaraTrace.h"
 #include "NiagaraWorldManager.h"
 #include "PrimitiveSceneInfo.h"
 #include "PrimitiveViewRelevance.h"
@@ -1287,6 +1288,8 @@ void UNiagaraComponent::ActivateInternal(bool bReset /* = false */, bool bIsScal
 	};
 	FScopedComponentTickEnabled ScopedComponentTickEnabled(this);
 
+	ON_SCOPE_EXIT { FNiagaraInsights::NotifyComponentActivate(this, bReset, bIsScalabilityCull, bAwaitingActivationDueToNotReady); };
+
 	bAwaitingActivationDueToNotReady = false;
 
 	// Reset our local bounds on reset
@@ -1561,6 +1564,8 @@ void UNiagaraComponent::Deactivate()
 
 void UNiagaraComponent::DeactivateInternal(bool bIsScalabilityCull /* = false */)
 {
+	FNiagaraInsights::NotifyComponentDeactivate(this, false, bIsScalabilityCull);
+
 	bool bWasCulledByScalabiltiy = bIsCulledByScalability;
 	bAwaitingActivationDueToNotReady = false;
 
@@ -1639,6 +1644,8 @@ bool UNiagaraComponent::IsComplete() const
 
 void UNiagaraComponent::DeactivateImmediateInternal(bool bIsScalabilityCull)
 {
+	FNiagaraInsights::NotifyComponentDeactivate(this, true, bIsScalabilityCull);
+
 	SCOPE_CYCLE_COUNTER(STAT_NiagaraComponentDeactivate);
 	CSV_SCOPED_TIMING_STAT_EXCLUSIVE(Effects);
 
@@ -1794,6 +1801,8 @@ void UNiagaraComponent::PostSystemTick_GameThread()
 
 void UNiagaraComponent::OnSystemComplete(bool bExternalCompletion)
 {
+	FNiagaraInsights::NotifyComponentComplete(this, bExternalCompletion);
+
 	// Debug feature, if we have loop enabled all world FX will loop unless deactivate by scalability
 #if !UE_BUILD_SHIPPING
 	if ( !bExternalCompletion && !bIsCulledByScalability )
diff --git a/Engine/Plugins/FX/Niagara/Source/Niagara/Private/NiagaraGPUProfiler.cpp b/Engine/Plugins/FX/Niagara/Source/Niagara/Private/NiagaraGPUProfiler.cpp
--- a/Engine/Plugins/FX/Niagara/Source/Niagara/Private/NiagaraGPUProfiler.cpp
+++ b/Engine/Plugins/FX/Niagara/Source/Niagara/Private/NiagaraGPUProfiler.cpp
@@ -3,6 +3,7 @@
 #include "NiagaraGPUProfiler.h"
 #include "HAL/IConsoleManager.h"
 #include "RHICommandList.h"
+#include "NiagaraTrace.h"
 
 bool GNiagaraGpuProfilingEnabled = true;
 static FAutoConsoleVariableRef CVarNiagaraGpuProfilingEnabled(
@@ -54,7 +55,8 @@ void FNiagaraGPUProfiler::BeginFrame(FRHICommandList& RHICmdList)
 	// If we are not enabled nothing to profile
 	static const auto CSVStatsEnabledCVar = IConsoleManager::Get().FindConsoleVariable(TEXT("fx.DetailedCSVStats"));
 	const bool bCsvCollecting = CSVStatsEnabledCVar && CSVStatsEnabledCVar->GetBool();
-	if ( !GNiagaraGpuProfilingEnabled || ((NumReaders == 0) && !bCsvCollecting) )
+	const bool bInsightsTraceCollecting = FNiagaraInsights::IsChannelActive();
+	if ( !GNiagaraGpuProfilingEnabled || ((NumReaders == 0) && !bCsvCollecting && !bInsightsTraceCollecting) )
 	{
 		return;
 	}
diff --git a/Engine/Plugins/FX/Niagara/Source/Niagara/Private/NiagaraModule.cpp b/Engine/Plugins/FX/Niagara/Source/Niagara/Private/NiagaraModule.cpp
--- a/Engine/Plugins/FX/Niagara/Source/Niagara/Private/NiagaraModule.cpp
+++ b/Engine/Plugins/FX/Niagara/Source/Niagara/Private/NiagaraModule.cpp
@@ -34,6 +34,7 @@
 #include "NiagaraCustomVersion.h"
 #include "NiagaraDataChannel.h"
 #include "NiagaraShaderModule.h"
+#include "NiagaraTrace.h"
 #include "NiagaraTypeRegistry.h"
 #include "NiagaraGpuComputeDispatch.h"
 #include "Misc/CoreDelegates.h"
@@ -628,6 +629,7 @@ void INiagaraModule::OnPostGarbageCollect()
 
 void INiagaraModule::OnBeginFrame()
 {
+	FNiagaraInsights::Update();
 	FNiagaraComponentSettings::UpdateSettings();
 	NiagaraStatelessCommon::UpdateSettings();
 	FNiagaraPlatformSet::RefreshScalability();
diff --git a/Engine/Plugins/FX/Niagara/Source/Niagara/Private/NiagaraTrace.cpp b/Engine/Plugins/FX/Niagara/Source/Niagara/Private/NiagaraTrace.cpp
--- a/Engine/Plugins/FX/Niagara/Source/Niagara/Private/NiagaraTrace.cpp
+++ b/Engine/Plugins/FX/Niagara/Source/Niagara/Private/NiagaraTrace.cpp
@@ -2,6 +2,238 @@
 
 #include "NiagaraTrace.h"
 
-#include "Trace/Detail/Channel.h"
+#include "NiagaraComponent.h"
 
-UE_TRACE_CHANNEL_DEFINE(NiagaraChannel);
+#include "Particles/ParticlePerfStats.h"
+
+UE_TRACE_CHANNEL_DEFINE(NiagaraChannel)
+
+#if WITH_NIAGARA_INSIGHTS
+namespace NiagaraInsightsPrivate
+{
+	void BuildComponentName(FNameBuilder& Builder, const UNiagaraComponent* Component)
+	{
+		if (AActor* OwnerActor = Component->GetTypedOuter<AActor>())
+		{
+			FStringView ActorLabel = OwnerActor->GetActorLabelView();
+			if (ActorLabel.IsEmpty())
+			{
+				OwnerActor->GetFName().AppendString(Builder);
+			}
+			else
+			{
+				Builder.Append(ActorLabel);
+			}
+			Builder.Append(TEXT("."));
+		}
+		Component->GetFName().AppendString(Builder);
+
+		if (UNiagaraSystem* System = Component->GetAsset())
+		{
+			Builder.Append(TEXT(","));
+			System->GetFName().AppendString(Builder);
+		}
+		else
+		{
+			Builder.Append(TEXT(",nullptr"));
+		}
+	}
+
+	UE_TRACE_EVENT_BEGIN(NiagaraTrace, SystemPerformance_GT, NoSync)
+		UE_TRACE_EVENT_FIELD(uint64,				EventTime)
+		UE_TRACE_EVENT_FIELD(UE::Trace::WideString, SystemName)
+		UE_TRACE_EVENT_FIELD(uint64,				NumInstances)
+		UE_TRACE_EVENT_FIELD(uint64,				TickGameThreadCycles)
+		UE_TRACE_EVENT_FIELD(uint64,				TickConcurrentCycles)		
+		UE_TRACE_EVENT_FIELD(uint64,				FinalizeCycles)
+		UE_TRACE_EVENT_FIELD(uint64,				EndOfFrameCycles)
+		UE_TRACE_EVENT_FIELD(uint64,				ActivationCycles)
+		UE_TRACE_EVENT_FIELD(uint64,				WaitCycles)
+		UE_TRACE_EVENT_FIELD(uint64,				MemoryBytes)
+	UE_TRACE_EVENT_END()
+
+	UE_TRACE_EVENT_BEGIN(NiagaraTrace, SystemPerformance_RT, NoSync)
+		UE_TRACE_EVENT_FIELD(uint64,				EventTime)
+		UE_TRACE_EVENT_FIELD(UE::Trace::WideString, SystemName)
+		UE_TRACE_EVENT_FIELD(uint64,				NumInstances)
+		UE_TRACE_EVENT_FIELD(uint64,				RenderUpdateCycles)
+		UE_TRACE_EVENT_FIELD(uint64,				GetDynamicMeshElementsCycles)
+		UE_TRACE_EVENT_FIELD(uint64,				GpuNumInstances)
+		UE_TRACE_EVENT_FIELD(uint64,				GpuTotalMicroseconds)
+	UE_TRACE_EVENT_END()
+
+	UE_TRACE_EVENT_BEGIN(NiagaraTrace, ComponentActivate, NoSync)
+		UE_TRACE_EVENT_FIELD(uint64,				EventTime)
+		UE_TRACE_EVENT_FIELD(UE::Trace::WideString, ComponentName)
+		UE_TRACE_EVENT_FIELD(bool,					bReset)
+		UE_TRACE_EVENT_FIELD(bool,					bIsScalabilityCull)
+		UE_TRACE_EVENT_FIELD(bool,					bAwaitingActivationDueToNotReady)
+	UE_TRACE_EVENT_END()
+
+	UE_TRACE_EVENT_BEGIN(NiagaraTrace, ComponentDeactivate, NoSync)
+		UE_TRACE_EVENT_FIELD(uint64,				EventTime)
+		UE_TRACE_EVENT_FIELD(UE::Trace::WideString, ComponentName)
+		UE_TRACE_EVENT_FIELD(bool,					bImmediate)
+		UE_TRACE_EVENT_FIELD(bool,					bIsScalabilityCull)
+	UE_TRACE_EVENT_END()
+
+	UE_TRACE_EVENT_BEGIN(NiagaraTrace, ComponentComplete, NoSync)
+		UE_TRACE_EVENT_FIELD(uint64,				EventTime)
+		UE_TRACE_EVENT_FIELD(UE::Trace::WideString, ComponentName)
+		UE_TRACE_EVENT_FIELD(bool,					bExternalCompletion)
+	UE_TRACE_EVENT_END()
+
+	class FNiagaraInsightsStatsListener : public FParticlePerfStatsListener
+	{
+	public:
+		virtual bool Tick() override
+		{
+			const uint64 CaptureTime = FPlatformTime::Cycles64();
+			FParticlePerfStatsManager::ForAllSystemStats(
+				[&](TWeakObjectPtr<const UFXSystemAsset>& WeakSystem, TUniquePtr<FParticlePerfStats>& Stats)
+				{
+					const FParticlePerfStats_GT& Stats_GT = Stats->GetGameThreadStats();
+					if (Stats_GT.NumInstances > 0)
+					{
+						if (const UNiagaraSystem* System = Cast<UNiagaraSystem>(WeakSystem.Get()))
+						{
+							FNameBuilder SystemNameBuilder;
+							System->GetFName().ToString(SystemNameBuilder);
+
+							UE_TRACE_LOG(NiagaraTrace, SystemPerformance_GT, NiagaraChannel)
+								<< SystemPerformance_GT.EventTime(CaptureTime)
+								<< SystemPerformance_GT.SystemName(SystemNameBuilder.ToString())
+								<< SystemPerformance_GT.NumInstances(Stats_GT.NumInstances)
+								<< SystemPerformance_GT.TickGameThreadCycles(Stats_GT.TickGameThreadCycles)
+								<< SystemPerformance_GT.TickConcurrentCycles(Stats_GT.TickConcurrentCycles.Load(EMemoryOrder::Relaxed))
+								<< SystemPerformance_GT.FinalizeCycles(Stats_GT.FinalizeCycles)
+								<< SystemPerformance_GT.EndOfFrameCycles(Stats_GT.EndOfFrameCycles.Load(EMemoryOrder::Relaxed))
+								<< SystemPerformance_GT.ActivationCycles(Stats_GT.ActivationCycles.Load(EMemoryOrder::Relaxed))
+								<< SystemPerformance_GT.WaitCycles(Stats_GT.WaitCycles)
+								<< SystemPerformance_GT.MemoryBytes(Stats_GT.MemoryBytes + Stats->MemoryKB_Asset.Get(0))
+								;
+						}
+					}
+				}
+			);
+
+			return true;
+		}
+
+		virtual void TickRT() override
+		{
+			const uint64 CaptureTime = FPlatformTime::Cycles64();
+			FParticlePerfStatsManager::ForAllSystemStats(
+				[&](TWeakObjectPtr<const UFXSystemAsset>& WeakSystem, TUniquePtr<FParticlePerfStats>& Stats)
+				{
+					const FParticlePerfStats_RT& Stats_RT = Stats->GetRenderThreadStats();
+					if (Stats_RT.NumInstances > 0)
+					{
+						if (const UNiagaraSystem* System = Cast<UNiagaraSystem>(WeakSystem.Get()))
+						{
+							FNameBuilder SystemNameBuilder;
+							System->GetFName().ToString(SystemNameBuilder);
+
+							const FParticlePerfStats_GPU& GPUStats = Stats->GetGPUStats();
+
+							UE_TRACE_LOG(NiagaraTrace, SystemPerformance_RT, NiagaraChannel)
+								<< SystemPerformance_RT.EventTime(CaptureTime)
+								<< SystemPerformance_RT.SystemName(SystemNameBuilder.ToString())
+								<< SystemPerformance_RT.NumInstances(Stats_RT.NumInstances)
+								<< SystemPerformance_RT.RenderUpdateCycles(Stats_RT.RenderUpdateCycles)
+								<< SystemPerformance_RT.GetDynamicMeshElementsCycles(Stats_RT.GetDynamicMeshElementsCycles)
+								<< SystemPerformance_RT.GpuNumInstances(GPUStats.NumInstances)
+								<< SystemPerformance_RT.GpuTotalMicroseconds(GPUStats.TotalMicroseconds)
+								;
+						}
+					}
+				}
+			);
+		}
+
+		virtual bool NeedsMemoryStats() const override { return true; }
+		virtual bool NeedsWorldStats() const override { return false; }
+		virtual bool NeedsSystemStats() const override { return true; }
+		virtual bool NeedsComponentStats() const override { return false; }
+	};
+
+	TSharedPtr<FNiagaraInsightsStatsListener> GStatsListener;
+}
+
+void FNiagaraInsights::Update()
+{
+	if (IsChannelActive())
+	{
+		if (NiagaraInsightsPrivate::GStatsListener == nullptr)
+		{
+			NiagaraInsightsPrivate::GStatsListener = MakeShared<NiagaraInsightsPrivate::FNiagaraInsightsStatsListener>();
+			FParticlePerfStatsManager::AddListener(NiagaraInsightsPrivate::GStatsListener);
+		}
+	}
+	else
+	{
+		if (NiagaraInsightsPrivate::GStatsListener != nullptr)
+		{
+			FParticlePerfStatsManager::RemoveListener(NiagaraInsightsPrivate::GStatsListener);
+			NiagaraInsightsPrivate::GStatsListener = nullptr;
+		}
+	}
+}
+
+bool FNiagaraInsights::IsChannelActive()
+{
+	return UE_TRACE_CHANNELEXPR_IS_ENABLED(NiagaraChannel);
+}
+
+void FNiagaraInsights::NotifyComponentActivate(const UNiagaraComponent* Component, bool bReset, bool bIsScalabilityCull, bool bAwaitingActivationDueToNotReady)
+{
+	using namespace NiagaraInsightsPrivate;
+	if (IsChannelActive())
+	{
+		FNameBuilder SystemNameBuilder;
+		NiagaraInsightsPrivate::BuildComponentName(SystemNameBuilder, Component);
+
+		UE_TRACE_LOG(NiagaraTrace, ComponentActivate, NiagaraChannel)
+			<< ComponentActivate.EventTime(FPlatformTime::Cycles64())
+			<< ComponentActivate.ComponentName(SystemNameBuilder.ToString())
+			<< ComponentActivate.bReset(bReset)
+			<< ComponentActivate.bIsScalabilityCull(bIsScalabilityCull)
+			<< ComponentActivate.bAwaitingActivationDueToNotReady(bAwaitingActivationDueToNotReady)
+			;
+	}
+}
+
+void FNiagaraInsights::NotifyComponentDeactivate(const UNiagaraComponent* Component, bool bImmediate, bool bIsScalabilityCull)
+{
+	using namespace NiagaraInsightsPrivate;
+	if (IsChannelActive())
+	{
+		FNameBuilder SystemNameBuilder;
+		NiagaraInsightsPrivate::BuildComponentName(SystemNameBuilder, Component);
+
+		UE_TRACE_LOG(NiagaraTrace, ComponentDeactivate, NiagaraChannel)
+			<< ComponentDeactivate.EventTime(FPlatformTime::Cycles64())
+			<< ComponentDeactivate.ComponentName(SystemNameBuilder.ToString())
+			<< ComponentDeactivate.bImmediate(bImmediate)
+			<< ComponentDeactivate.bIsScalabilityCull(bIsScalabilityCull)
+			;
+	}
+}
+
+void FNiagaraInsights::NotifyComponentComplete(const UNiagaraComponent* Component, bool bExternalCompletion)
+{
+	using namespace NiagaraInsightsPrivate;
+	if (IsChannelActive())
+	{
+		FNameBuilder SystemNameBuilder;
+		NiagaraInsightsPrivate::BuildComponentName(SystemNameBuilder, Component);
+
+		UE_TRACE_LOG(NiagaraTrace, ComponentComplete, NiagaraChannel)
+			<< ComponentComplete.EventTime(FPlatformTime::Cycles64())
+			<< ComponentComplete.ComponentName(SystemNameBuilder.ToString())
+			<< ComponentComplete.bExternalCompletion(bExternalCompletion)
+			;
+	}
+}
+
+#endif //WITH_NIAGARA_INSIGHTS
diff --git a/Engine/Plugins/FX/Niagara/Source/Niagara/Public/NiagaraTrace.h b/Engine/Plugins/FX/Niagara/Source/Niagara/Public/NiagaraTrace.h
deleted file mode 100644
diff --git a/Engine/Plugins/FX/NiagaraInsights/NiagaraInsights.uplugin b/Engine/Plugins/FX/NiagaraInsights/NiagaraInsights.uplugin
new file mode 100644
--- /dev/null
+++ b/Engine/Plugins/FX/NiagaraInsights/NiagaraInsights.uplugin
@@ -0,0 +1,28 @@
+{
+	"FileVersion": 3,
+	"Version": 1,
+	"VersionName": "1.0",
+	"FriendlyName": "Niagara Insights",
+	"Description": "Niagara Unreal Insights Debugging and Performance",
+	"Category": "Insights",
+	"CreatedBy": "Epic Games, Inc.",
+	"CreatedByURL": "https://epicgames.com",
+	"DocsURL": "",
+	"MarketplaceURL": "",
+	"SupportURL": "",
+	"EnabledByDefault": true,
+	"CanContainContent": false,
+	"IsBetaVersion": false,
+	"Installed": false,
+	"SupportedPrograms": [ "UnrealInsights" ],
+	"Modules": [
+		{
+			"Name": "NiagaraInsights",
+			"Type": "EditorAndProgram",
+			"LoadingPhase": "Default",
+			"ProgramAllowList": [
+				"UnrealInsights"
+			]
+		}
+	]
+}
diff --git a/Engine/Plugins/FX/NiagaraInsights/Source/NiagaraInsights/NiagaraInsights.build.cs b/Engine/Plugins/FX/NiagaraInsights/Source/NiagaraInsights/NiagaraInsights.build.cs
new file mode 100644
--- /dev/null
+++ b/Engine/Plugins/FX/NiagaraInsights/Source/NiagaraInsights/NiagaraInsights.build.cs
@@ -0,0 +1,26 @@
+// Copyright Epic Games, Inc. All Rights Reserved.
+
+namespace UnrealBuildTool.Rules
+{
+	public class NiagaraInsights : ModuleRules
+	{
+		public NiagaraInsights(ReadOnlyTargetRules Target) : base(Target)
+		{
+			PrivateDependencyModuleNames.AddRange(
+			new string[]
+			{
+				"Core",
+				"CoreUObject",
+				"SlateCore",
+				"Slate",
+				"TraceLog",
+				"TraceAnalysis",
+				"TraceServices",
+				"TraceInsights",
+				"TraceInsightsCore",
+				"InputCore"
+			});
+		}
+	}
+}
+
diff --git a/Engine/Plugins/FX/NiagaraInsights/Source/NiagaraInsights/Private/NiagaraAnalyzer.cpp b/Engine/Plugins/FX/NiagaraInsights/Source/NiagaraInsights/Private/NiagaraAnalyzer.cpp
new file mode 100644
--- /dev/null
+++ b/Engine/Plugins/FX/NiagaraInsights/Source/NiagaraInsights/Private/NiagaraAnalyzer.cpp
@@ -0,0 +1,117 @@
+// Copyright Epic Games, Inc. All Rights Reserved.
+
+#include "NiagaraAnalyzer.h"
+#include "NiagaraProvider.h"
+
+#include "TraceServices/Model/AnalysisSession.h"
+
+#define LOCTEXT_NAMESPACE "NiagaraProvider"
+
+namespace UE::NiagaraInsights
+{
+
+static FStringView GetEventString(const Trace::IAnalyzer::FEventData& EventData, const ANSICHAR* FieldName)
+{
+	FStringView StringView;
+	EventData.GetString(FieldName, StringView);
+	return StringView;
+}
+
+FNiagaraAnalyzer::FNiagaraAnalyzer(TraceServices::IAnalysisSession& InSession, FNiagaraProvider& InProvider)
+	: Session(InSession)
+	, Provider(InProvider)
+{
+}
+
+void FNiagaraAnalyzer::OnAnalysisBegin(const FOnAnalysisContext& Context)
+{
+	UE::Trace::IAnalyzer::FInterfaceBuilder& Builder = Context.InterfaceBuilder;
+	
+	Builder.RouteEvent(RouteId_SystemPerformance_GT, "NiagaraTrace", "SystemPerformance_GT");
+	Builder.RouteEvent(RouteId_SystemPerformance_RT, "NiagaraTrace", "SystemPerformance_RT");
+
+	Builder.RouteEvent(RouteId_ComponentActivate,	"NiagaraTrace", "ComponentActivate");
+	Builder.RouteEvent(RouteId_ComponentDeactivate,	"NiagaraTrace", "ComponentDeactivate");
+	Builder.RouteEvent(RouteId_ComponentComplete,	"NiagaraTrace", "ComponentComplete");
+}
+
+bool FNiagaraAnalyzer::OnEvent(uint16 RouteId, EStyle Style, const FOnEventContext& Context)
+{
+	//LLM_SCOPE_BYNAME(TEXT("Insights/FNiagaraAnalyzer"));
+
+	const Trace::IAnalyzer::FEventData& EventData = Context.EventData;
+	const double EventTime = Context.EventTime.AsSeconds(EventData.GetValue<uint64>("EventTime"));
+
+	TraceServices::FAnalysisSessionEditScope _(Session);
+	switch (RouteId)
+	{
+		case RouteId_SystemPerformance_GT:
+		{
+			Provider.AddSystemPerformance_GT(
+				EventTime,
+				GetEventString(EventData, "SystemName"),
+				EventData.GetValue<uint64>("NumInstances"),
+				Context.EventTime.AsSecondsAbsolute(EventData.GetValue<uint64>("TickGameThreadCycles")),
+				Context.EventTime.AsSecondsAbsolute(EventData.GetValue<uint64>("TickConcurrentCycles")),
+				Context.EventTime.AsSecondsAbsolute(EventData.GetValue<uint64>("FinalizeCycles")),
+				Context.EventTime.AsSecondsAbsolute(EventData.GetValue<uint64>("EndOfFrameCycles")),
+				Context.EventTime.AsSecondsAbsolute(EventData.GetValue<uint64>("ActivationCycles")),
+				Context.EventTime.AsSecondsAbsolute(EventData.GetValue<uint64>("WaitCycles")),
+				EventData.GetValue<uint64>("MemoryBytes")
+			);
+			break;
+		}
+
+		case RouteId_SystemPerformance_RT:
+		{
+			Provider.AddSystemPerformance_RT(
+				EventTime,
+				GetEventString(EventData, "SystemName"),
+				EventData.GetValue<uint64>("NumInstances"),
+				Context.EventTime.AsSecondsAbsolute(EventData.GetValue<uint64>("RenderUpdateCycles")),
+				Context.EventTime.AsSecondsAbsolute(EventData.GetValue<uint64>("GetDynamicMeshElementsCycles")),
+				EventData.GetValue<uint64>("GpuNumInstances"),
+				EventData.GetValue<uint64>("GpuTotalMicroseconds")
+			);
+			break;
+		}
+
+		case RouteId_ComponentActivate:
+		{
+			Provider.AddComponentActivate(
+				EventTime,
+				GetEventString(EventData, "ComponentName"),
+				EventData.GetValue<bool>("bReset"),
+				EventData.GetValue<bool>("bIsScalabilityCull"),
+				EventData.GetValue<bool>("bAwaitingActivationDueToNotReady")
+			);
+			break;
+		}
+
+		case RouteId_ComponentDeactivate:
+		{
+			Provider.AddComponentDeactivate(
+				EventTime,
+				GetEventString(EventData, "ComponentName"),
+				EventData.GetValue<bool>("bImmediate"),
+				EventData.GetValue<bool>("bIsScalabilityCull")
+			);
+			break;
+		}
+
+		case RouteId_ComponentComplete:
+		{
+			Provider.AddComponentComplete(
+				EventTime,
+				GetEventString(EventData, "ComponentName"),
+				EventData.GetValue<bool>("bExternalCompletion")
+			);
+			break;
+		}
+	}
+	return true;
+}
+
+} //namespace UE::NiagaraInsights
+
+#undef LOCTEXT_NAMESPACE
diff --git a/Engine/Plugins/FX/NiagaraInsights/Source/NiagaraInsights/Private/NiagaraAnalyzer.h b/Engine/Plugins/FX/NiagaraInsights/Source/NiagaraInsights/Private/NiagaraAnalyzer.h
new file mode 100644
--- /dev/null
+++ b/Engine/Plugins/FX/NiagaraInsights/Source/NiagaraInsights/Private/NiagaraAnalyzer.h
@@ -0,0 +1,38 @@
+// Copyright Epic Games, Inc. All Rights Reserved.
+
+#pragma once
+
+#include "Trace/Analyzer.h"
+
+namespace TraceServices { class IAnalysisSession; }
+
+namespace UE::NiagaraInsights
+{
+
+class FNiagaraProvider;
+
+class FNiagaraAnalyzer : public UE::Trace::IAnalyzer
+{
+public:
+	FNiagaraAnalyzer(TraceServices::IAnalysisSession& InSession, FNiagaraProvider& InProvider);
+
+	virtual void OnAnalysisBegin(const FOnAnalysisContext& Context) override;
+	virtual void OnAnalysisEnd() override {}
+	virtual bool OnEvent(uint16 RouteId, EStyle Style, const FOnEventContext& Context) override;
+
+private:
+	enum : uint16
+	{
+		RouteId_SystemPerformance_GT,
+		RouteId_SystemPerformance_RT,
+
+		RouteId_ComponentActivate,
+		RouteId_ComponentDeactivate,
+		RouteId_ComponentComplete,
+	};
+
+	TraceServices::IAnalysisSession&	Session;
+	FNiagaraProvider&					Provider;
+};
+
+} //namespace UE::NiagaraInsights
diff --git a/Engine/Plugins/FX/NiagaraInsights/Source/NiagaraInsights/Private/NiagaraInsightsCommon.cpp b/Engine/Plugins/FX/NiagaraInsights/Source/NiagaraInsights/Private/NiagaraInsightsCommon.cpp
new file mode 100644
--- /dev/null
+++ b/Engine/Plugins/FX/NiagaraInsights/Source/NiagaraInsights/Private/NiagaraInsightsCommon.cpp
@@ -0,0 +1,16 @@
+// Copyright Epic Games, Inc. All Rights Reserved.
+
+#include "NiagaraInsightsCommon.h"
+
+namespace UE::NiagaraInsights::Common
+{
+
+const FLinearColor Color_InstanceCount(0.10f, 0.45f, 0.80f, 1.f);	// Blue
+const FLinearColor Color_GTCostTotal(0.85f, 0.25f, 0.15f, 1.f);		// Red
+const FLinearColor Color_GTCostPerInst(1.00f, 0.55f, 0.15f, 1.f);	// Orange
+const FLinearColor Color_RTCostTotal(0.15f, 0.70f, 0.40f, 1.f);		// Green
+const FLinearColor Color_RTCostPerInst(0.60f, 0.85f, 0.30f, 1.f);	// Yellow-green
+const FLinearColor Color_GpuCostTotal(0.65f, 0.25f, 0.85f, 1.f);	// Purple
+const FLinearColor Color_MemoryCostTotal(0.10f, 0.45f, 0.80f, 1.f);	// Purple
+
+} // namespace UE::NiagaraInsights
diff --git a/Engine/Plugins/FX/NiagaraInsights/Source/NiagaraInsights/Private/NiagaraInsightsCommon.h b/Engine/Plugins/FX/NiagaraInsights/Source/NiagaraInsights/Private/NiagaraInsightsCommon.h
new file mode 100644
--- /dev/null
+++ b/Engine/Plugins/FX/NiagaraInsights/Source/NiagaraInsights/Private/NiagaraInsightsCommon.h
@@ -0,0 +1,18 @@
+// Copyright Epic Games, Inc. All Rights Reserved.
+
+#pragma once
+
+#include "CoreMinimal.h"
+
+namespace UE::NiagaraInsights::Common
+{
+
+extern const FLinearColor Color_InstanceCount;
+extern const FLinearColor Color_GTCostTotal;
+extern const FLinearColor Color_GTCostPerInst;
+extern const FLinearColor Color_RTCostTotal;
+extern const FLinearColor Color_RTCostPerInst;
+extern const FLinearColor Color_GpuCostTotal;
+extern const FLinearColor Color_MemoryCostTotal;
+
+} // namespace UE::NiagaraInsights::Common
diff --git a/Engine/Plugins/FX/NiagaraInsights/Source/NiagaraInsights/Private/NiagaraInsightsComponent.cpp b/Engine/Plugins/FX/NiagaraInsights/Source/NiagaraInsights/Private/NiagaraInsightsComponent.cpp
new file mode 100644
--- /dev/null
+++ b/Engine/Plugins/FX/NiagaraInsights/Source/NiagaraInsights/Private/NiagaraInsightsComponent.cpp
@@ -0,0 +1,85 @@
+// Copyright Epic Games, Inc. All Rights Reserved.
+
+#include "NiagaraInsightsComponent.h"
+#include "Widgets/SNiagaraRangeStatsView.h"
+
+#include "Framework/Docking/LayoutExtender.h"
+#include "Framework/Docking/TabManager.h"
+#include "Insights/IUnrealInsightsModule.h"
+#include "Modules/ModuleManager.h"
+#include "Widgets/Docking/SDockTab.h"
+
+#define LOCTEXT_NAMESPACE "NiagaraInsightsComponent"
+
+namespace UE::NiagaraInsights
+{
+
+const FName FNiagaraInsightsTabs::RangeStatsViewID("NiagaraRangeStats");
+
+////////////////////////////////////////////////////////////////////////////////////////////////////
+
+void FNiagaraInsightsComponent::Initialize(IUnrealInsightsModule& InsightsModule)
+{
+	// Nothing needed at init time; tab registration happens via RegisterMajorTabs.
+}
+
+void FNiagaraInsightsComponent::Shutdown()
+{
+}
+
+void FNiagaraInsightsComponent::RegisterMajorTabs(IUnrealInsightsModule& InsightsModule)
+{
+	// Extend the Timing Profiler major tab to include our minor tab.
+	FOnRegisterMajorTabExtensions& Delegate = InsightsModule.OnRegisterMajorTabExtension(FInsightsManagerTabs::TimingProfilerTabId);
+
+	ExtensionDelegateHandle = Delegate.AddRaw(this, &FNiagaraInsightsComponent::OnRegisterTimingProfilerTabExtensions);
+}
+
+void FNiagaraInsightsComponent::UnregisterMajorTabs()
+{
+	if (ExtensionDelegateHandle.IsValid() && FModuleManager::Get().IsModuleLoaded("TraceInsights"))
+	{
+		IUnrealInsightsModule& InsightsModule = FModuleManager::LoadModuleChecked<IUnrealInsightsModule>("TraceInsights");
+		FOnRegisterMajorTabExtensions* Delegate = InsightsModule.FindMajorTabLayoutExtension(FInsightsManagerTabs::TimingProfilerTabId);
+		if (Delegate)
+		{
+			Delegate->Remove(ExtensionDelegateHandle);
+		}
+	}
+	ExtensionDelegateHandle.Reset();
+}
+
+////////////////////////////////////////////////////////////////////////////////////////////////////
+
+void FNiagaraInsightsComponent::OnRegisterTimingProfilerTabExtensions(FInsightsMajorTabExtender& Extender)
+{
+	// Register the minor tab spawner with the Timing Profiler's tab manager.
+	FMinorTabConfig& TabConfig = Extender.AddMinorTabConfig();
+	TabConfig.TabId      = FNiagaraInsightsTabs::RangeStatsViewID;
+	TabConfig.TabLabel   = LOCTEXT("NiagaraStatsTabTitle", "Niagara Stats");
+	TabConfig.TabTooltip = LOCTEXT("NiagaraStatsTabTooltip",
+		"Shows top-5 Niagara systems by various metrics for the selected time range.");
+	TabConfig.OnSpawnTab = FOnSpawnTab::CreateRaw(this, &FNiagaraInsightsComponent::SpawnTab_RangeStatsView);
+
+	// Place the tab after the existing "Timers" tab in the Timing Profiler layout.
+	// The user can also drag it anywhere in the Timing Profiler window.
+	Extender.GetLayoutExtender().ExtendLayout(
+		FTimingProfilerTabs::TimersID,
+		ELayoutExtensionPosition::After,
+		FTabManager::FTab(FNiagaraInsightsTabs::RangeStatsViewID, ETabState::OpenedTab)
+	);
+}
+
+TSharedRef<SDockTab> FNiagaraInsightsComponent::SpawnTab_RangeStatsView(const FSpawnTabArgs& Args)
+{
+	return SNew(SDockTab)
+		.ShouldAutosize(false)
+		.TabRole(ETabRole::PanelTab)
+		[
+			SNew(SNiagaraRangeStatsView)
+		];
+}
+
+} // namespace UE::NiagaraInsights
+
+#undef LOCTEXT_NAMESPACE
diff --git a/Engine/Plugins/FX/NiagaraInsights/Source/NiagaraInsights/Private/NiagaraInsightsComponent.h b/Engine/Plugins/FX/NiagaraInsights/Source/NiagaraInsights/Private/NiagaraInsightsComponent.h
new file mode 100644
--- /dev/null
+++ b/Engine/Plugins/FX/NiagaraInsights/Source/NiagaraInsights/Private/NiagaraInsightsComponent.h
@@ -0,0 +1,39 @@
+// Copyright Epic Games, Inc. All Rights Reserved.
+
+#pragma once
+
+#include "Insights/IUnrealInsightsModule.h"
+
+namespace UE::NiagaraInsights
+{
+
+/** Tab IDs for the Niagara Insights minor tabs inside the Timing Profiler. */
+struct FNiagaraInsightsTabs
+{
+	static const FName RangeStatsViewID;
+};
+
+/**
+ * IInsightsComponent that adds the Niagara Range Stats minor tab into the Timing Profiler.
+ * The tab shows top-5 offenders by instance count, GT cost, GT cost/instance, RT cost,
+ * RT cost/instance, and GPU cost for the currently selected time range.
+ */
+class FNiagaraInsightsComponent : public IInsightsComponent
+{
+public:
+	//~ Begin IInsightsComponent interface
+	virtual void Initialize(IUnrealInsightsModule& InsightsModule) override;
+	virtual void Shutdown() override;
+	virtual void RegisterMajorTabs(IUnrealInsightsModule& InsightsModule) override;
+	virtual void UnregisterMajorTabs() override;
+	//~ End IInsightsComponent interface
+
+private:
+	void OnRegisterTimingProfilerTabExtensions(FInsightsMajorTabExtender& Extender);
+
+	TSharedRef<SDockTab> SpawnTab_RangeStatsView(const FSpawnTabArgs& Args);
+
+	FDelegateHandle ExtensionDelegateHandle;
+};
+
+} // namespace UE::NiagaraInsights
diff --git a/Engine/Plugins/FX/NiagaraInsights/Source/NiagaraInsights/Private/NiagaraInsightsModule.cpp b/Engine/Plugins/FX/NiagaraInsights/Source/NiagaraInsights/Private/NiagaraInsightsModule.cpp
new file mode 100644
--- /dev/null
+++ b/Engine/Plugins/FX/NiagaraInsights/Source/NiagaraInsights/Private/NiagaraInsightsModule.cpp
@@ -0,0 +1,79 @@
+// Copyright Epic Games, Inc. All Rights Reserved.
+
+#include "NiagaraInsightsModule.h"
+#include "NiagaraInsightsComponent.h"
+
+#include "Features/IModularFeatures.h"
+#include "Insights/IUnrealInsightsModule.h"
+#include "Modules/ModuleManager.h"
+
+namespace UE::NiagaraInsights
+{
+
+void FNiagaraInsightsModule::StartupModule()
+{
+	IModularFeatures::Get().RegisterModularFeature(TraceServices::ModuleFeatureName, &TraceModule);
+	IModularFeatures::Get().RegisterModularFeature(UE::Insights::Timing::TimingViewExtenderFeatureName, &TimingViewExtender);
+
+	// Try to register with TraceInsights now; if it isn't loaded yet, listen for when it loads.
+	TryRegisterInsightsComponent();
+	if (!InsightsComponent.IsValid())
+	{
+		ModuleChangedHandle = FModuleManager::Get().OnModulesChanged().AddRaw(
+			this, &FNiagaraInsightsModule::OnModulesChanged);
+	}
+}
+
+void FNiagaraInsightsModule::TryRegisterInsightsComponent()
+{
+	if (InsightsComponent.IsValid())
+	{
+		return;
+	}
+	if (!FModuleManager::Get().IsModuleLoaded("TraceInsights"))
+	{
+		return;
+	}
+	IUnrealInsightsModule& InsightsModule = FModuleManager::LoadModuleChecked<IUnrealInsightsModule>("TraceInsights");
+	InsightsComponent = MakeShared<FNiagaraInsightsComponent>();
+	InsightsModule.RegisterComponent(InsightsComponent);
+}
+
+void FNiagaraInsightsModule::OnModulesChanged(FName ModuleName, EModuleChangeReason Reason)
+{
+	if (ModuleName == "TraceInsights" && Reason == EModuleChangeReason::ModuleLoaded)
+	{
+		TryRegisterInsightsComponent();
+		// No longer need to listen once registered.
+		FModuleManager::Get().OnModulesChanged().Remove(ModuleChangedHandle);
+		ModuleChangedHandle.Reset();
+	}
+}
+
+void FNiagaraInsightsModule::ShutdownModule()
+{
+	if (ModuleChangedHandle.IsValid())
+	{
+		FModuleManager::Get().OnModulesChanged().Remove(ModuleChangedHandle);
+		ModuleChangedHandle.Reset();
+	}
+
+	if (InsightsComponent.IsValid() && FModuleManager::Get().IsModuleLoaded("TraceInsights"))
+	{
+		IUnrealInsightsModule& InsightsModule = FModuleManager::LoadModuleChecked<IUnrealInsightsModule>("TraceInsights");
+		InsightsModule.UnregisterComponent(InsightsComponent);
+	}
+	InsightsComponent.Reset();
+
+	IModularFeatures::Get().UnregisterModularFeature(TraceServices::ModuleFeatureName, &TraceModule);
+	IModularFeatures::Get().UnregisterModularFeature(UE::Insights::Timing::TimingViewExtenderFeatureName, &TimingViewExtender);
+}
+
+FNiagaraInsightsModule& FNiagaraInsightsModule::Get()
+{
+	return FModuleManager::LoadModuleChecked<FNiagaraInsightsModule>("NiagaraInsights");
+}
+
+} //namespace UE::NiagaraInsights
+
+IMPLEMENT_MODULE(UE::NiagaraInsights::FNiagaraInsightsModule, NiagaraInsights);
diff --git a/Engine/Plugins/FX/NiagaraInsights/Source/NiagaraInsights/Private/NiagaraInsightsModule.h b/Engine/Plugins/FX/NiagaraInsights/Source/NiagaraInsights/Private/NiagaraInsightsModule.h
new file mode 100644
--- /dev/null
+++ b/Engine/Plugins/FX/NiagaraInsights/Source/NiagaraInsights/Private/NiagaraInsightsModule.h
@@ -0,0 +1,37 @@
+// Copyright Epic Games, Inc. All Rights Reserved.
+
+#pragma once
+
+#include "Modules/ModuleInterface.h"
+#include "Modules/ModuleManager.h"
+#include "NiagaraTraceModule.h"
+#include "NiagaraTimingViewExtender.h"
+
+namespace UE::NiagaraInsights
+{
+
+class FNiagaraInsightsComponent;
+
+class FNiagaraInsightsModule : public IModuleInterface
+{
+public:
+	//~ Begin IModuleInterface interface
+	virtual void StartupModule() override;
+	virtual void ShutdownModule() override;
+	//~ End IModuleInterface interface
+
+	static FNiagaraInsightsModule& Get();
+
+	FNiagaraTimingViewExtender& GetTimingViewExtender() { return TimingViewExtender; }
+
+private:
+	void TryRegisterInsightsComponent();
+	void OnModulesChanged(FName ModuleName, EModuleChangeReason Reason);
+
+	FNiagaraTraceModule				TraceModule;
+	FNiagaraTimingViewExtender		TimingViewExtender;
+	TSharedPtr<FNiagaraInsightsComponent> InsightsComponent;
+	FDelegateHandle ModuleChangedHandle;
+};
+
+} //namespace UE::NiagaraInsights
diff --git a/Engine/Plugins/FX/NiagaraInsights/Source/NiagaraInsights/Private/NiagaraInstanceLifecycleTrack.cpp b/Engine/Plugins/FX/NiagaraInsights/Source/NiagaraInsights/Private/NiagaraInstanceLifecycleTrack.cpp
new file mode 100644
--- /dev/null
+++ b/Engine/Plugins/FX/NiagaraInsights/Source/NiagaraInsights/Private/NiagaraInstanceLifecycleTrack.cpp
@@ -0,0 +1,144 @@
+// Copyright Epic Games, Inc. All Rights Reserved.
+
+#include "NiagaraInstanceLifecycleTrack.h"
+
+#include "NiagaraInsightsCommon.h"
+#include "NiagaraProvider.h"
+#include "NiagaraTimingViewSession.h"
+
+#include "Insights/ViewModels/GraphTrackBuilder.h"
+#include "Insights/ViewModels/GraphTrackEvent.h"
+#include "Insights/ViewModels/TimingTrackViewport.h"
+#include "Insights/ViewModels/TooltipDrawState.h"
+
+#define LOCTEXT_NAMESPACE "NiagaraInstanceLifecycleTrack"
+
+namespace UE::NiagaraInsights
+{
+
+INSIGHTS_IMPLEMENT_RTTI(FNiagaraInstanceLifecycleTrack)
+
+FNiagaraInstanceLifecycleTrack::FNiagaraInstanceLifecycleTrack(const FNiagaraTimingViewSession& InSharedData)
+	: FGraphTrack(LOCTEXT("NiagaraLifecycle", "Niagara Lifecycle").ToString())
+	, SharedData(InSharedData)
+{
+	constexpr double NumEventTypes = 3.0;
+
+	EnabledOptions = //EGraphOptions::ShowDebugInfo |
+		EGraphOptions::ShowPoints |
+		EGraphOptions::ShowLines |
+		EGraphOptions::ShowPolygon |
+		EGraphOptions::ShowVerticalAxisGrid |
+		EGraphOptions::ShowHeader;
+
+	SetValidLocations(ETimingTrackLocation::Scrollable | ETimingTrackLocation::TopDocked | ETimingTrackLocation::BottomDocked);
+	SetHeight(NumEventTypes * 10.0);
+
+	TSharedRef<FGraphSeries> GraphSeries = MakeShared<FGraphSeries>();
+	GraphSeries->SetName(TEXT("Events"));
+	GraphSeries->SetColor(Common::Color_GTCostTotal);
+	GraphSeries->SetBaselineY(GetHeight());
+	GraphSeries->SetScaleY(GetHeight() / (NumEventTypes + 1.0));
+	GraphSeries->EnableAutoZoom();
+	GetSeries().Add(GraphSeries);
+}
+
+void FNiagaraInstanceLifecycleTrack::Update(const ITimingTrackUpdateContext& Context)
+{
+	Super::Update(Context);
+
+	const FNiagaraProvider* Provider = SharedData.GetAnalysisSession().ReadProvider<FNiagaraProvider>(FNiagaraProvider::GetProviderName());
+	if (Provider == nullptr)
+	{
+		return;
+	}
+
+	const FTimingTrackViewport& Viewport = Context.GetViewport();
+
+	const TSharedPtr<FGraphSeries> GraphSeries = GetSeries()[0];
+	if (GraphSeries->IsVisible())
+	{
+		FGraphTrackBuilder Builder(*this, *GraphSeries, Viewport);
+		Provider->EnumerateComponentEvent(
+			Viewport.GetStartTime(),
+			Viewport.GetEndTime(),
+			[&Builder, &GraphSeries](const FLifetimeEvent& LifetimeEvent)
+			{
+				const double Offset = double(LifetimeEvent.Payload.GetIndex()) + 0.5;
+				Builder.AddEvent(LifetimeEvent.Time, 0.0, Offset, false);
+			}
+		);
+	}
+}
+
+void FNiagaraInstanceLifecycleTrack::InitTooltip(FTooltipDrawState& InOutTooltip, const ITimingEvent& InTooltipEvent) const
+{
+	if (!InTooltipEvent.CheckTrack(this) || !InTooltipEvent.Is<FGraphTrackEvent>())
+	{
+		return;
+	}
+
+	const FGraphTrackEvent& GraphEvent = InTooltipEvent.As<FGraphTrackEvent>();
+	const TSharedRef<const FGraphSeries> HoveredSeries = GraphEvent.GetSeries();
+
+	const FNiagaraProvider* Provider = SharedData.GetAnalysisSession().ReadProvider<FNiagaraProvider>(FNiagaraProvider::GetProviderName());
+	if (Provider == nullptr)
+	{
+		return;
+	}
+
+	InOutTooltip.ResetContent();
+	InOutTooltip.AddTitle(HoveredSeries->GetName().ToString(), HoveredSeries->GetColor());
+
+	const double EventTime = GraphEvent.GetStartTime();
+	constexpr double TimeTolerance = 1e-9;
+
+	Provider->EnumerateComponentEvent(
+		EventTime - TimeTolerance,
+		EventTime + TimeTolerance,
+		[&InOutTooltip](const FLifetimeEvent& LifetimeEvent)
+		{
+			switch (LifetimeEvent.Payload.GetIndex())
+			{
+				case 0:
+				{
+					const FLifetimeEvent::FActivate& Payload = LifetimeEvent.Payload.Get<FLifetimeEvent::FActivate>();
+					InOutTooltip.AddNameValueTextLine(TEXT("Activate"), TEXT(""));
+					InOutTooltip.AddNameValueTextLine(TEXT("Component"), LifetimeEvent.ComponentName);
+					InOutTooltip.AddNameValueTextLine(TEXT("System"), LifetimeEvent.SystemName);
+					InOutTooltip.AddNameValueTextLine(TEXT("bReset"), LexToString(Payload.bReset));
+					InOutTooltip.AddNameValueTextLine(TEXT("bIsScalabilityCull"), LexToString(Payload.bIsScalabilityCull));
+					InOutTooltip.AddNameValueTextLine(TEXT("bAwaitingActivationDueToNotReady"), LexToString(Payload.bAwaitingActivationDueToNotReady));
+					break;
+				}
+
+				case 1:
+				{
+					const FLifetimeEvent::FDeactivate& Payload = LifetimeEvent.Payload.Get<FLifetimeEvent::FDeactivate>();
+					InOutTooltip.AddNameValueTextLine(TEXT("DeActivate"), TEXT(""));
+					InOutTooltip.AddNameValueTextLine(TEXT("Component"), LifetimeEvent.ComponentName);
+					InOutTooltip.AddNameValueTextLine(TEXT("System"), LifetimeEvent.SystemName);
+					InOutTooltip.AddNameValueTextLine(TEXT("bImmediate"), LexToString(Payload.bImmediate));
+					InOutTooltip.AddNameValueTextLine(TEXT("bIsScalabilityCull"), LexToString(Payload.bIsScalabilityCull));
+					break;
+				}
+
+				case 2:
+				{
+					const FLifetimeEvent::FComplete& Payload = LifetimeEvent.Payload.Get<FLifetimeEvent::FComplete>();
+					InOutTooltip.AddNameValueTextLine(TEXT("Complete"), TEXT(""));
+					InOutTooltip.AddNameValueTextLine(TEXT("Component"), LifetimeEvent.ComponentName);
+					InOutTooltip.AddNameValueTextLine(TEXT("System"), LifetimeEvent.SystemName);
+					InOutTooltip.AddNameValueTextLine(TEXT("bExternalCompletion"), LexToString(Payload.bExternalCompletion));
+					break;
+				}
+			}
+		}
+	);
+
+	InOutTooltip.UpdateLayout();
+}
+
+} //namespace UE::NiagaraInsights
+
+#undef LOCTEXT_NAMESPACE
diff --git a/Engine/Plugins/FX/NiagaraInsights/Source/NiagaraInsights/Private/NiagaraInstanceLifecycleTrack.h b/Engine/Plugins/FX/NiagaraInsights/Source/NiagaraInsights/Private/NiagaraInstanceLifecycleTrack.h
new file mode 100644
--- /dev/null
+++ b/Engine/Plugins/FX/NiagaraInsights/Source/NiagaraInsights/Private/NiagaraInstanceLifecycleTrack.h
@@ -0,0 +1,30 @@
+// Copyright Epic Games, Inc. All Rights Reserved.
+
+#pragma once
+
+#include "Insights/ViewModels/GraphTrack.h"
+
+namespace UE::NiagaraInsights
+{
+
+class FNiagaraTimingViewSession;
+
+class FNiagaraInstanceLifecycleTrack : public FGraphTrack
+{
+	using Super = FGraphTrack;
+
+	INSIGHTS_DECLARE_RTTI(FNiagaraInstanceLifecycleTrack, FGraphTrack)
+
+public:
+	explicit FNiagaraInstanceLifecycleTrack(const FNiagaraTimingViewSession& InSharedData);
+
+	// BEGIN: FGraphTrack interface
+	virtual void Update(const ITimingTrackUpdateContext& Context) override;
+	virtual void InitTooltip(FTooltipDrawState& InOutTooltip, const ITimingEvent& InTooltipEvent) const override;
+	// END: FGraphTrack interface
+
+private:
+	const FNiagaraTimingViewSession& SharedData;
+};
+
+} //namespace UE::NiagaraInsights
diff --git a/Engine/Plugins/FX/NiagaraInsights/Source/NiagaraInsights/Private/NiagaraPerformanceGraphTrack.cpp b/Engine/Plugins/FX/NiagaraInsights/Source/NiagaraInsights/Private/NiagaraPerformanceGraphTrack.cpp
new file mode 100644
--- /dev/null
+++ b/Engine/Plugins/FX/NiagaraInsights/Source/NiagaraInsights/Private/NiagaraPerformanceGraphTrack.cpp
@@ -0,0 +1,391 @@
+// Copyright Epic Games, Inc. All Rights Reserved.
+
+#include "NiagaraPerformanceGraphTrack.h"
+
+#include "NiagaraInsightsCommon.h"
+#include "NiagaraProvider.h"
+#include "NiagaraTimingViewSession.h"
+
+#include "Insights/ViewModels/GraphTrackBuilder.h"
+#include "Insights/ViewModels/GraphTrackEvent.h"
+#include "Insights/ViewModels/TimingTrackViewport.h"
+#include "Insights/ViewModels/TooltipDrawState.h"
+
+#define LOCTEXT_NAMESPACE "NiagaraPerformanceGraphTrack"
+
+namespace UE::NiagaraInsights
+{
+
+INSIGHTS_IMPLEMENT_RTTI(FNiagaraPerformanceGraphTrack)
+
+FNiagaraPerformanceGraphTrack::FNiagaraPerformanceGraphTrack(const FNiagaraTimingViewSession& InSharedData)
+	: FGraphTrack(LOCTEXT("NiagaraPerformance", "Niagara Performance").ToString())
+	, SharedData(InSharedData)
+{
+	EnabledOptions = //EGraphOptions::ShowDebugInfo |
+		EGraphOptions::ShowPoints |
+		EGraphOptions::ShowLines |
+		EGraphOptions::ShowPolygon |
+		EGraphOptions::ShowVerticalAxisGrid |
+		EGraphOptions::ShowHeader;
+
+	SetValidLocations(ETimingTrackLocation::Scrollable | ETimingTrackLocation::TopDocked | ETimingTrackLocation::BottomDocked);
+	SetHeight(200.0f);	// Fixed height for graph
+
+	//-TOOD: Crate series for metrics
+	{
+		GameTimeSeries = MakeShared<FGraphSeries>();
+		GameTimeSeries->SetName(TEXT("Game Time"));
+		GameTimeSeries->SetColor(Common::Color_GTCostTotal);
+		GameTimeSeries->EnableAutoZoom();
+		GetSeries().Add(GameTimeSeries);
+	}
+	{
+		RenderTimeSeries = MakeShared<FGraphSeries>();
+		RenderTimeSeries->SetName(TEXT("Render Time"));
+		RenderTimeSeries->SetColor(Common::Color_RTCostTotal);
+		RenderTimeSeries->EnableAutoZoom();
+		GetSeries().Add(RenderTimeSeries);
+	}
+	{
+		GpuTimeSeries = MakeShared<FGraphSeries>();
+		GpuTimeSeries->SetName(TEXT("GPU Time"));
+		GpuTimeSeries->SetColor(Common::Color_GpuCostTotal);
+		GpuTimeSeries->EnableAutoZoom();
+		GetSeries().Add(GpuTimeSeries);
+	}
+	{
+		MemoryUsageSeries = MakeShared<FGraphSeries>();
+		MemoryUsageSeries->SetName(TEXT("Memory Usage"));
+		MemoryUsageSeries->SetColor(Common::Color_MemoryCostTotal);
+		MemoryUsageSeries->EnableAutoZoom();
+		GetSeries().Add(MemoryUsageSeries);
+	}
+}
+
+void FNiagaraPerformanceGraphTrack::Update(const ITimingTrackUpdateContext& Context)
+{
+	Super::Update(Context);
+
+	const FNiagaraProvider* Provider = SharedData.GetAnalysisSession().ReadProvider<FNiagaraProvider>(FNiagaraProvider::GetProviderName());
+	if (Provider == nullptr)
+	{
+		return;
+	}
+
+	const FTimingTrackViewport& Viewport = Context.GetViewport();
+
+	// Update height based on visibility
+	float SeriesHeight = GetHeight();
+	{
+		int NumVisible = 0;
+		for (TSharedPtr<FGraphSeries>& Series : GetSeries())
+		{
+			NumVisible += Series->IsVisible() ? 1 : 0;
+		}
+
+		if (NumVisible == 0)
+		{
+			return;
+		}
+
+		SeriesHeight = SeriesHeight / double(NumVisible);
+		float CurrentSeriesHeight = SeriesHeight;
+		for (TSharedPtr<FGraphSeries>& Series : GetSeries())
+		{
+			if (Series->IsVisible())
+			{
+				Series->SetBaselineY(CurrentSeriesHeight);
+				CurrentSeriesHeight += SeriesHeight;
+			}
+		}
+
+		// Add on a few pixel so we don't calculate the top point on the line of the next series
+		SeriesHeight -= 5.0;
+	}
+
+	// Update game time series
+	if (GameTimeSeries->IsVisible())
+	{
+		double MaxValue = 0.0;
+		FGraphTrackBuilder Builder(*this, *GameTimeSeries, Viewport);
+		Provider->EnumeratePerformance_GT(
+			Viewport.GetStartTime(),
+			Viewport.GetEndTime(),
+			[&Builder, &MaxValue](const FSystemPerformanceFrame_GT& Frame)
+			{
+				double TotalTime = 0.0;
+				for (const FSystemPerformanceFrame_GT::FSystemData& SystemData : Frame.SystemData)
+				{
+					TotalTime +=
+						SystemData.TickGameThreadSeconds +
+						SystemData.TickConcurrentSeconds +
+						SystemData.FinalizeSeconds +
+						SystemData.EndOfFrameSeconds +
+						SystemData.ActivationSeconds +
+						SystemData.WaitSeconds;
+
+				}
+				const double Value = TotalTime * 1000.0;
+				MaxValue = FMath::Max(MaxValue, Value);
+				Builder.AddEvent(Frame.Time, 0.0, Value);
+			}
+		);
+		GameTimeSeries->SetScaleY(MaxValue > 0.0 ? SeriesHeight / MaxValue : 1.0);
+	}
+
+	// Update render time series
+	if (RenderTimeSeries->IsVisible())
+	{
+		double MaxValue = 0.0;
+		FGraphTrackBuilder Builder(*this, *RenderTimeSeries, Viewport);
+		Provider->EnumeratePerformance_RT(
+			Viewport.GetStartTime(),
+			Viewport.GetEndTime(),
+			[&Builder, &MaxValue](const FSystemPerformanceFrame_RT& Frame)
+			{
+				double TotalTime = 0.0;
+				for (const FSystemPerformanceFrame_RT::FSystemData& SystemData : Frame.SystemData)
+				{
+					TotalTime += SystemData.RenderUpdateSeconds + SystemData.GetDynamicMeshElementsSeconds;
+				}
+				const double Value = TotalTime * 1000.0;
+				MaxValue = FMath::Max(MaxValue, Value);
+				Builder.AddEvent(Frame.Time, 0.0, Value);
+			}
+		);
+		RenderTimeSeries->SetScaleY(MaxValue > 0.0 ? SeriesHeight / MaxValue : 1.0);
+	}
+
+	// Update gpu time series
+	if (GpuTimeSeries->IsVisible())
+	{
+		double MaxValue = 0.0;
+		FGraphTrackBuilder Builder(*this, *GpuTimeSeries, Viewport);
+		Provider->EnumeratePerformance_RT(
+			Viewport.GetStartTime(),
+			Viewport.GetEndTime(),
+			[&Builder, &MaxValue](const FSystemPerformanceFrame_RT& Frame)
+			{
+				double TotalTime = 0.0;
+				for (const FSystemPerformanceFrame_RT::FSystemData& SystemData : Frame.SystemData)
+				{
+					TotalTime += SystemData.GpuTotalMicroseconds;
+				}
+				const double Value = TotalTime / 1000.0;
+				MaxValue = FMath::Max(MaxValue, Value);
+				Builder.AddEvent(Frame.Time, 0.0, Value);
+			}
+		);
+		GpuTimeSeries->SetScaleY(MaxValue > 0.0 ? SeriesHeight / MaxValue : 1.0);
+	}
+
+	// Update memory series
+	if (MemoryUsageSeries->IsVisible())
+	{
+		double MaxValue = 0.0;
+		FGraphTrackBuilder Builder(*this, *MemoryUsageSeries, Viewport);
+		Provider->EnumeratePerformance_GT(
+			Viewport.GetStartTime(),
+			Viewport.GetEndTime(),
+			[&Builder, &MaxValue](const FSystemPerformanceFrame_GT& Frame)
+			{
+				uint64 MemoryTotal = 0;
+				for (const FSystemPerformanceFrame_GT::FSystemData& SystemData : Frame.SystemData)
+				{
+					MemoryTotal += SystemData.MemoryBytes;
+
+				}
+				const double Value = double(MemoryTotal) / 1024.0;
+				MaxValue = FMath::Max(MaxValue, Value);
+				Builder.AddEvent(Frame.Time, 0.0, Value);
+			}
+		);
+		MemoryUsageSeries->SetScaleY(MaxValue > 0.0 ? SeriesHeight / MaxValue : 1.0);
+	}
+}
+
+void FNiagaraPerformanceGraphTrack::InitTooltip(FTooltipDrawState& InOutTooltip, const ITimingEvent& InTooltipEvent) const
+{
+	if (!InTooltipEvent.CheckTrack(this) || !InTooltipEvent.Is<FGraphTrackEvent>())
+	{
+		return;
+	}
+
+	const FGraphTrackEvent& GraphEvent = InTooltipEvent.As<FGraphTrackEvent>();
+	const TSharedRef<const FGraphSeries> HoveredSeries = GraphEvent.GetSeries();
+
+	const FNiagaraProvider* Provider = SharedData.GetAnalysisSession().ReadProvider<FNiagaraProvider>(FNiagaraProvider::GetProviderName());
+	if (Provider == nullptr)
+	{
+		return;
+	}
+
+	InOutTooltip.ResetContent();
+	InOutTooltip.AddTitle(HoveredSeries->GetName().ToString(), HoveredSeries->GetColor());
+
+	// Use a tight window to find the single frame this event was built from.
+	// Frames are ~16ms apart so 1ns tolerance should match at most one.
+	const double EventTime = GraphEvent.GetStartTime();
+	constexpr double TimeTolerance = 1e-9;
+
+	if (HoveredSeries == GameTimeSeries)
+	{
+		Provider->EnumeratePerformance_GT(EventTime - TimeTolerance, EventTime + TimeTolerance,
+			[&InOutTooltip](const FSystemPerformanceFrame_GT& Frame)
+			{
+				struct FSystemEntry { FStringView Name; double TotalMs; uint64 NumInstances; };
+
+				TArray<FSystemEntry> Entries;
+				double GrandTotalMs = 0.0;
+
+				for (const FSystemPerformanceFrame_GT::FSystemData& Sys : Frame.SystemData)
+				{
+					const double TotalMs = (Sys.TickGameThreadSeconds + Sys.TickConcurrentSeconds +
+						Sys.FinalizeSeconds + Sys.EndOfFrameSeconds +
+						Sys.ActivationSeconds + Sys.WaitSeconds) * 1000.0;
+					Entries.Add({ Sys.SystemName, TotalMs, Sys.NumInstances });
+					GrandTotalMs += TotalMs;
+				}
+
+				Entries.Sort([](const FSystemEntry& A, const FSystemEntry& B) { return A.TotalMs > B.TotalMs; });
+
+				InOutTooltip.AddNameValueTextLine(TEXT("Total:"), FString::Printf(TEXT("%.3f ms"), GrandTotalMs));
+				InOutTooltip.AddTextLine(TEXT("Per System (sorted by cost):"), FLinearColor::Gray);
+
+				constexpr int32 MaxSystems = 10;
+				for (int32 i = 0; i < FMath::Min(Entries.Num(), MaxSystems); ++i)
+				{
+					const FSystemEntry& Entry = Entries[i];
+					const FString Label = FString::Printf(TEXT("  %s (x%llu)"), *FString(Entry.Name), Entry.NumInstances);
+					InOutTooltip.AddNameValueTextLine(*Label, FString::Printf(TEXT("%.3f ms"), Entry.TotalMs));
+				}
+
+				if (Entries.Num() > MaxSystems)
+				{
+					InOutTooltip.AddTextLine(FString::Printf(TEXT("  ... and %d more"), Entries.Num() - MaxSystems), FLinearColor::Gray);
+				}
+			}
+		);
+	}
+	else if (HoveredSeries == RenderTimeSeries)
+	{
+		Provider->EnumeratePerformance_RT(EventTime - TimeTolerance, EventTime + TimeTolerance,
+			[&InOutTooltip](const FSystemPerformanceFrame_RT& Frame)
+			{
+				struct FSystemEntry { FStringView Name; double TotalMs; uint64 NumInstances; };
+
+				TArray<FSystemEntry> Entries;
+				double GrandTotalMs = 0.0;
+
+				for (const FSystemPerformanceFrame_RT::FSystemData& Sys : Frame.SystemData)
+				{
+					const double TotalMs = (Sys.RenderUpdateSeconds + Sys.GetDynamicMeshElementsSeconds) * 1000.0;
+					Entries.Add({ Sys.SystemName, TotalMs, Sys.NumInstances });
+					GrandTotalMs += TotalMs;
+				}
+
+				Entries.Sort([](const FSystemEntry& A, const FSystemEntry& B) { return A.TotalMs > B.TotalMs; });
+
+				InOutTooltip.AddNameValueTextLine(TEXT("Total:"), FString::Printf(TEXT("%.3f ms"), GrandTotalMs));
+				InOutTooltip.AddTextLine(TEXT("Per System (sorted by cost):"), FLinearColor::Gray);
+
+				constexpr int32 MaxSystems = 10;
+				for (int32 i = 0; i < FMath::Min(Entries.Num(), MaxSystems); ++i)
+				{
+					const FSystemEntry& Entry = Entries[i];
+					const FString Label = FString::Printf(TEXT("  %s (x%llu)"), *FString(Entry.Name), Entry.NumInstances);
+					InOutTooltip.AddNameValueTextLine(*Label, FString::Printf(TEXT("%.3f ms"), Entry.TotalMs));
+				}
+
+				if (Entries.Num() > MaxSystems)
+				{
+					InOutTooltip.AddTextLine(FString::Printf(TEXT("  ... and %d more"), Entries.Num() - MaxSystems), FLinearColor::Gray);
+				}
+			}
+		);
+	}
+	else if (HoveredSeries == GpuTimeSeries)
+	{
+		Provider->EnumeratePerformance_RT(EventTime - TimeTolerance, EventTime + TimeTolerance,
+			[&InOutTooltip](const FSystemPerformanceFrame_RT& Frame)
+			{
+				struct FSystemEntry { FStringView Name; double TotalMs; uint64 GpuNumInstances; };
+
+				TArray<FSystemEntry> Entries;
+				double GrandTotalMs = 0.0;
+
+				for (const FSystemPerformanceFrame_RT::FSystemData& Sys : Frame.SystemData)
+				{
+					if (Sys.GpuNumInstances == 0)
+					{
+						continue;
+					}
+					const double TotalMs = Sys.GpuTotalMicroseconds / 1000.0;
+					Entries.Add({ Sys.SystemName, TotalMs, Sys.GpuNumInstances });
+					GrandTotalMs += TotalMs;
+				}
+
+				Entries.Sort([](const FSystemEntry& A, const FSystemEntry& B) { return A.TotalMs > B.TotalMs; });
+
+				InOutTooltip.AddNameValueTextLine(TEXT("Total:"), FString::Printf(TEXT("%.3f ms"), GrandTotalMs));
+				InOutTooltip.AddTextLine(TEXT("Per System (sorted by cost):"), FLinearColor::Gray);
+
+				constexpr int32 MaxSystems = 10;
+				for (int32 i = 0; i < FMath::Min(Entries.Num(), MaxSystems); ++i)
+				{
+					const FSystemEntry& Entry = Entries[i];
+					const FString Label = FString::Printf(TEXT("  %s (x%llu)"), *FString(Entry.Name), Entry.GpuNumInstances);
+					InOutTooltip.AddNameValueTextLine(*Label, FString::Printf(TEXT("%.3f ms"), Entry.TotalMs));
+				}
+
+				if (Entries.Num() > MaxSystems)
+				{
+					InOutTooltip.AddTextLine(FString::Printf(TEXT("  ... and %d more"), Entries.Num() - MaxSystems), FLinearColor::Gray);
+				}
+			}
+		);
+	}
+	else if (HoveredSeries == MemoryUsageSeries)
+	{
+		Provider->EnumeratePerformance_GT(EventTime - TimeTolerance, EventTime + TimeTolerance,
+			[&InOutTooltip](const FSystemPerformanceFrame_GT& Frame)
+			{
+				struct FSystemEntry { FStringView Name; uint64 MemoryTotalBytes; uint64 NumInstances; };
+
+				TArray<FSystemEntry> Entries;
+				uint64 MemoryTotalBytes = 0;
+
+				for (const FSystemPerformanceFrame_GT::FSystemData& Sys : Frame.SystemData)
+				{
+					Entries.Add({ Sys.SystemName, Sys.MemoryBytes, Sys.NumInstances });
+					MemoryTotalBytes += Sys.MemoryBytes;
+				}
+
+				Entries.Sort([](const FSystemEntry& A, const FSystemEntry& B) { return A.MemoryTotalBytes > B.MemoryTotalBytes; });
+
+				InOutTooltip.AddNameValueTextLine(TEXT("Total:"), FString::Printf(TEXT("%.3f KB"), double(MemoryTotalBytes) / 1024.0));
+				InOutTooltip.AddTextLine(TEXT("Per System (sorted by cost):"), FLinearColor::Gray);
+
+				constexpr int32 MaxSystems = 10;
+				for (int32 i = 0; i < FMath::Min(Entries.Num(), MaxSystems); ++i)
+				{
+					const FSystemEntry& Entry = Entries[i];
+					const FString Label = FString::Printf(TEXT("  %s (x%llu)"), *FString(Entry.Name), Entry.NumInstances);
+					InOutTooltip.AddNameValueTextLine(*Label, FString::Printf(TEXT("%.3f KB"), double(Entry.MemoryTotalBytes) / 1024.0));
+				}
+
+				if (Entries.Num() > MaxSystems)
+				{
+					InOutTooltip.AddTextLine(FString::Printf(TEXT("  ... and %d more"), Entries.Num() - MaxSystems), FLinearColor::Gray);
+				}
+			}
+		);
+	}
+
+	InOutTooltip.UpdateLayout();
+}
+
+} //namespace UE::NiagaraInsights
+
+#undef LOCTEXT_NAMESPACE
diff --git a/Engine/Plugins/FX/NiagaraInsights/Source/NiagaraInsights/Private/NiagaraPerformanceGraphTrack.h b/Engine/Plugins/FX/NiagaraInsights/Source/NiagaraInsights/Private/NiagaraPerformanceGraphTrack.h
new file mode 100644
--- /dev/null
+++ b/Engine/Plugins/FX/NiagaraInsights/Source/NiagaraInsights/Private/NiagaraPerformanceGraphTrack.h
@@ -0,0 +1,36 @@
+// Copyright Epic Games, Inc. All Rights Reserved.
+
+#pragma once
+
+#include "Insights/ViewModels/GraphTrack.h"
+
+class FGraphSeries;
+
+namespace UE::NiagaraInsights
+{
+
+class FNiagaraTimingViewSession;
+
+class FNiagaraPerformanceGraphTrack : public FGraphTrack
+{
+	using Super = FGraphTrack;
+
+	INSIGHTS_DECLARE_RTTI(FNiagaraPerformanceGraphTrack, FGraphTrack)
+
+public:
+	explicit FNiagaraPerformanceGraphTrack(const FNiagaraTimingViewSession& InSharedData);
+
+	// BEGIN: FGraphTrack interface
+	virtual void Update(const ITimingTrackUpdateContext& Context) override;
+	virtual void InitTooltip(FTooltipDrawState& InOutTooltip, const ITimingEvent& InTooltipEvent) const override;
+	// END: FGraphTrack interface
+
+private:
+	const FNiagaraTimingViewSession&	SharedData;
+	TSharedPtr<FGraphSeries>			GameTimeSeries;
+	TSharedPtr<FGraphSeries>			RenderTimeSeries;
+	TSharedPtr<FGraphSeries>			GpuTimeSeries;
+	TSharedPtr<FGraphSeries>			MemoryUsageSeries;
+};
+
+} //namespace UE::NiagaraInsights
diff --git a/Engine/Plugins/FX/NiagaraInsights/Source/NiagaraInsights/Private/NiagaraProvider.cpp b/Engine/Plugins/FX/NiagaraInsights/Source/NiagaraInsights/Private/NiagaraProvider.cpp
new file mode 100644
--- /dev/null
+++ b/Engine/Plugins/FX/NiagaraInsights/Source/NiagaraInsights/Private/NiagaraProvider.cpp
@@ -0,0 +1,191 @@
+// Copyright Epic Games, Inc. All Rights Reserved.
+
+#include "NiagaraProvider.h"
+
+namespace UE::NiagaraInsights
+{
+namespace Private
+{
+	thread_local TraceServices::FProviderLock::FThreadLocalState GProviderLockState;
+	static FName ProviderName("NiagaraProvider");
+
+	void ParseComponentName(FLifetimeEvent& Event, FStringView ComponentName)
+	{
+		const int32 Length = ComponentName.Len();
+		int32 FoundIndex = INDEX_NONE;
+		if (ComponentName.FindChar(',', FoundIndex) == false)
+		{
+			FoundIndex = Length;
+		}
+
+		Event.ComponentName = ComponentName.Mid(0, FoundIndex);
+		Event.SystemName = ComponentName.Mid(FoundIndex + 1, FMath::Max(0, Length - FoundIndex - 1));
+	}
+}
+
+FNiagaraProvider::FNiagaraProvider(TraceServices::IAnalysisSession& InSession)
+	: Session(InSession)
+{
+}
+
+void FNiagaraProvider::BeginRead() const
+{
+	ProviderLock.BeginRead(Private::GProviderLockState);
+}
+
+void FNiagaraProvider::EndRead() const
+{
+	ProviderLock.EndRead(Private::GProviderLockState);
+}
+
+void FNiagaraProvider::ReadAccessCheck() const
+{
+	ProviderLock.ReadAccessCheck(Private::GProviderLockState);
+}
+
+void FNiagaraProvider::BeginEdit() const 
+{
+	ProviderLock.BeginWrite(Private::GProviderLockState);
+}
+
+void FNiagaraProvider::EndEdit() const
+{
+	ProviderLock.EndWrite(Private::GProviderLockState);
+}
+
+void FNiagaraProvider::EditAccessCheck() const
+{
+	ProviderLock.WriteAccessCheck(Private::GProviderLockState);
+}
+
+FName FNiagaraProvider::GetProviderName()
+{
+	return Private::ProviderName;
+}
+
+void FNiagaraProvider::AddSystemPerformance_GT(double EventTime, FStringView SystemName, uint64 NumInstances, double TickGameThreadSeconds, double TickConcurrentSeconds, double FinalizeSeconds, double EndOfFrameSeconds, double ActivationSeconds, double WaitSeconds, uint64 MemoryBytes)
+{
+	TraceServices::FProviderEditScopeLock WriteLock(*this);
+
+	//int32 iFrame = Algo::BinarySearchBy(PerformanceFrames, EventTime, &FPerformanceFrame::Time, FMath::IsNearlyEqual);
+	if (PerformanceFrames_GT.Num() == 0 || !FMath::IsNearlyEqual(PerformanceFrames_GT.Last().Time, EventTime))
+	{
+		PerformanceFrames_GT.AddDefaulted_GetRef().Time = EventTime;
+	}
+
+	FSystemPerformanceFrame_GT::FSystemData& SystemData = PerformanceFrames_GT.Last().SystemData.AddDefaulted_GetRef();
+	SystemData.SystemName = SystemName;
+	SystemData.NumInstances = NumInstances;
+	SystemData.TickGameThreadSeconds = TickGameThreadSeconds;
+	SystemData.TickConcurrentSeconds = TickConcurrentSeconds;
+	SystemData.FinalizeSeconds = FinalizeSeconds;
+	SystemData.EndOfFrameSeconds = EndOfFrameSeconds;
+	SystemData.ActivationSeconds = ActivationSeconds;
+	SystemData.WaitSeconds = WaitSeconds;
+	SystemData.MemoryBytes = MemoryBytes;
+}
+
+void FNiagaraProvider::AddSystemPerformance_RT(double EventTime, FStringView SystemName, uint64 NumInstances, double RenderUpdateSeconds, double GetDynamicMeshElementsSeconds, uint64 GpuNumInstances, uint64 GpuTotalMicroseconds)
+{
+	TraceServices::FProviderEditScopeLock WriteLock(*this);
+
+	//int32 iFrame = Algo::BinarySearchBy(PerformanceFrames, EventTime, &FPerformanceFrame::Time, FMath::IsNearlyEqual);
+	if (PerformanceFrames_RT.Num() == 0 || !FMath::IsNearlyEqual(PerformanceFrames_RT.Last().Time, EventTime))
+	{
+		PerformanceFrames_RT.AddDefaulted_GetRef().Time = EventTime;
+	}
+
+	FSystemPerformanceFrame_RT::FSystemData& SystemData = PerformanceFrames_RT.Last().SystemData.AddDefaulted_GetRef();
+	SystemData.SystemName = SystemName;
+	SystemData.NumInstances = NumInstances;
+	SystemData.RenderUpdateSeconds = RenderUpdateSeconds;
+	SystemData.GetDynamicMeshElementsSeconds = GetDynamicMeshElementsSeconds;
+	SystemData.GpuNumInstances = GpuNumInstances;
+	SystemData.GpuTotalMicroseconds = GpuTotalMicroseconds;
+}
+
+void FNiagaraProvider::EnumeratePerformance_GT(double StartTime, double EndTime, TFunction<void(const FSystemPerformanceFrame_GT&)> Evaluate) const
+{
+	TraceServices::FProviderReadScopeLock ReadLock(*this);
+
+	for (const FSystemPerformanceFrame_GT& Frame : PerformanceFrames_GT)
+	{
+		if (Frame.Time > EndTime)
+		{
+			return;
+		}
+		if (Frame.Time < StartTime)
+		{
+			continue;
+		}
+		Evaluate(Frame);
+	}
+}
+
+void FNiagaraProvider::EnumeratePerformance_RT(double StartTime, double EndTime, TFunction<void(const FSystemPerformanceFrame_RT&)> Evaluate) const
+{
+	TraceServices::FProviderReadScopeLock ReadLock(*this);
+
+	for (const FSystemPerformanceFrame_RT& Frame : PerformanceFrames_RT)
+	{
+		if (Frame.Time > EndTime)
+		{
+			return;
+		}
+		if (Frame.Time < StartTime)
+		{
+			continue;
+		}
+		Evaluate(Frame);
+	}
+}
+
+void FNiagaraProvider::AddComponentActivate(double EventTime, FStringView ComponentName, bool bReset, bool bIsScalabilityCull, bool bAwaitingActivationDueToNotReady)
+{
+	TraceServices::FProviderEditScopeLock WriteLock(*this);
+
+	FLifetimeEvent& Event = LifetimeEvents.AddDefaulted_GetRef();
+	Event.Time = EventTime;
+	Private::ParseComponentName(Event, ComponentName);
+	Event.Payload.Emplace<FLifetimeEvent::FActivate>(bReset, bIsScalabilityCull, bAwaitingActivationDueToNotReady);
+}
+
+void FNiagaraProvider::AddComponentDeactivate(double EventTime, FStringView ComponentName, bool bImmediate, bool bIsScalabilityCull)
+{
+	TraceServices::FProviderEditScopeLock WriteLock(*this);
+
+	FLifetimeEvent& Event = LifetimeEvents.AddDefaulted_GetRef();
+	Event.Time = EventTime;
+	Private::ParseComponentName(Event, ComponentName);
+	Event.Payload.Emplace<FLifetimeEvent::FDeactivate>(bImmediate, bIsScalabilityCull);
+}
+
+void FNiagaraProvider::AddComponentComplete(double EventTime, FStringView ComponentName, bool bExternalCompletion)
+{
+	TraceServices::FProviderEditScopeLock WriteLock(*this);
+
+	FLifetimeEvent& Event = LifetimeEvents.AddDefaulted_GetRef();
+	Event.Time = EventTime;
+	Private::ParseComponentName(Event, ComponentName);
+	Event.Payload.Emplace<FLifetimeEvent::FComplete>(bExternalCompletion);
+}
+
+void FNiagaraProvider::EnumerateComponentEvent(double StartTime, double EndTime, TFunction<void(const FLifetimeEvent&)> Evaluate) const
+{
+	TraceServices::FProviderReadScopeLock ReadLock(*this);
+
+	for (const FLifetimeEvent& Event : LifetimeEvents)
+	{
+		if (Event.Time > EndTime)
+		{
+			return;
+		}
+		if (Event.Time < StartTime)
+		{
+			continue;
+		}
+		Evaluate(Event);
+	}
+}
+
+} //namespace UE::NiagaraInsights
diff --git a/Engine/Plugins/FX/NiagaraInsights/Source/NiagaraInsights/Private/NiagaraProvider.h b/Engine/Plugins/FX/NiagaraInsights/Source/NiagaraInsights/Private/NiagaraProvider.h
new file mode 100644
--- /dev/null
+++ b/Engine/Plugins/FX/NiagaraInsights/Source/NiagaraInsights/Private/NiagaraProvider.h
@@ -0,0 +1,121 @@
+// Copyright Epic Games, Inc. All Rights Reserved.
+
+#pragma once
+
+#include "Common/ProviderLock.h"
+#include "InsightsCore/Common/SimpleRtti.h"
+#include "Model/IntervalTimeline.h"
+#include "Trace/Analyzer.h"
+#include "TraceServices/Model/AnalysisSession.h"
+
+namespace TraceServices { class IAnalysisSession; }
+
+namespace UE::NiagaraInsights
+{
+struct FSystemPerformanceFrame_GT
+{
+	struct FSystemData
+	{
+		FString SystemName;
+		uint64 NumInstances = 0;
+		double TickGameThreadSeconds = 0.0;
+		double TickConcurrentSeconds = 0.0;
+		double FinalizeSeconds = 0.0;
+		double EndOfFrameSeconds = 0.0;
+		double ActivationSeconds = 0.0;
+		double WaitSeconds = 0.0;
+		uint64 MemoryBytes = 0.0;
+	};
+
+	double Time = 0.0;
+	TArray<FSystemData> SystemData;
+};
+
+struct FSystemPerformanceFrame_RT
+{
+	struct FSystemData
+	{
+		FString SystemName;
+		uint64 NumInstances = 0;
+		double RenderUpdateSeconds = 0.0;
+		double GetDynamicMeshElementsSeconds = 0.0;
+
+		uint64 GpuNumInstances = 0;
+		uint64 GpuTotalMicroseconds = 0;
+	};
+
+	double Time = 0.0;
+	TArray<FSystemData> SystemData;
+};
+
+struct FLifetimeEvent
+{
+	struct FActivate
+	{
+		bool bReset = false;
+		bool bIsScalabilityCull = false;
+		bool bAwaitingActivationDueToNotReady = false;
+	};
+	struct FDeactivate
+	{
+		bool bImmediate = false;
+		bool bIsScalabilityCull = false;
+	};
+	struct FComplete
+	{
+		bool bExternalCompletion = false;
+	};
+
+	double Time = 0.0;
+	FString ComponentName;
+	FString SystemName;
+	TVariant<FActivate, FDeactivate, FComplete> Payload;
+};
+
+class FNiagaraProvider : public TraceServices::IProvider, public TraceServices::IEditableProvider
+{
+public:
+	explicit FNiagaraProvider(TraceServices::IAnalysisSession& InSession);
+
+	// Begin: TraceServices::IProvider Impl
+	virtual void BeginRead() const override;
+	virtual void EndRead() const override;
+	virtual void ReadAccessCheck() const override;
+	// End: TraceServices::IProvider Impl
+
+	// Begin: TraceServices::IEditableProvider Impl
+	virtual void BeginEdit() const override;
+	virtual void EndEdit() const override;
+	virtual void EditAccessCheck() const override;
+	// End: TraceServices::IEditableProvider Impl
+
+	static FName GetProviderName();
+
+	// Add system performance data for the game thread
+	void AddSystemPerformance_GT(double EventTime, FStringView SystemName, uint64 NumInstances, double TickGameThreadSeconds, double TickConcurrentSeconds, double FinalizeSeconds, double EndOfFrameSeconds, double ActivationSeconds, double WaitSeconds, uint64 MemoryBytes);
+	// Add system performnace data for the render thread
+	void AddSystemPerformance_RT(double EventTime, FStringView SystemName, uint64 NumInstances, double RenderUpdateSeconds, double GetDynamicMeshElementsSeconds, uint64 GpuNumInstances, uint64 GpuTotalMicroseconds);
+
+	// Enumerate a range of game thread frames
+	void EnumeratePerformance_GT(double StartTime, double EndTime, TFunction<void(const FSystemPerformanceFrame_GT&)> Evaluate) const;
+	// Enumerate a range of render thread frames
+	void EnumeratePerformance_RT(double StartTime, double EndTime, TFunction<void(const FSystemPerformanceFrame_RT&)> Evaluate) const;
+
+	void AddComponentActivate(double EventTime, FStringView ComponentName, bool bReset, bool bIsScalabilityCull, bool bAwaitingActivationDueToNotReady);
+	void AddComponentDeactivate(double EventTime, FStringView ComponentName, bool bImmediate, bool bIsScalabilityCull);
+	void AddComponentComplete(double EventTime, FStringView ComponentName, bool bExternalCompletion);
+
+	void EnumerateComponentEvent(double StartTime, double EndTime, TFunction<void(const FLifetimeEvent&)> Evaluate) const;
+
+private:
+	mutable TraceServices::FProviderLock	ProviderLock;
+
+	TraceServices::IAnalysisSession&		Session;
+
+	TArray<FSystemPerformanceFrame_GT>		PerformanceFrames_GT;
+	TArray<FSystemPerformanceFrame_RT>		PerformanceFrames_RT;
+
+	TArray<FLifetimeEvent>					LifetimeEvents;
+};
+
+} //namespace UE::NiagaraInsights
diff --git a/Engine/Plugins/FX/NiagaraInsights/Source/NiagaraInsights/Private/NiagaraTimingViewExtender.cpp b/Engine/Plugins/FX/NiagaraInsights/Source/NiagaraInsights/Private/NiagaraTimingViewExtender.cpp
new file mode 100644
--- /dev/null
+++ b/Engine/Plugins/FX/NiagaraInsights/Source/NiagaraInsights/Private/NiagaraTimingViewExtender.cpp
@@ -0,0 +1,79 @@
+// Copyright Epic Games, Inc. All Rights Reserved.
+
+#include "NiagaraTimingViewExtender.h"
+#include "Insights/ITimingViewSession.h"
+#include "Insights/IUnrealInsightsModule.h"
+
+#define LOCTEXT_NAMESPACE "NiagaraTimingViewExtender"
+
+namespace UE::NiagaraInsights
+{
+
+void FNiagaraTimingViewExtender::OnBeginSession(UE::Insights::Timing::ITimingViewSession& InSession)
+{
+	if (InSession.GetName() != FInsightsManagerTabs::TimingProfilerTabId)
+	{
+		return;
+	}
+
+	FPerSessionData* PerSessionData = PerSessionDataMap.Find(&InSession);
+	if (PerSessionData == nullptr)
+	{
+		PerSessionData = &PerSessionDataMap.Add(&InSession);
+		PerSessionData->SharedData = MakeUnique<FNiagaraTimingViewSession>();
+	}
+
+	PerSessionData->SharedData->OnBeginSession(InSession);
+	OnSessionBegun.Broadcast(*PerSessionData->SharedData);
+}
+
+FNiagaraTimingViewSession* FNiagaraTimingViewExtender::FindFirstActiveSession()
+{
+	for (auto& KV : PerSessionDataMap)
+	{
+		if (KV.Value.SharedData.IsValid())
+		{
+			return KV.Value.SharedData.Get();
+		}
+	}
+	return nullptr;
+}
+
+void FNiagaraTimingViewExtender::OnEndSession(UE::Insights::Timing::ITimingViewSession& InSession)
+{
+	if (InSession.GetName() != FInsightsManagerTabs::TimingProfilerTabId)
+	{
+		return;
+	}
+
+	FPerSessionData* PerSessionData = PerSessionDataMap.Find(&InSession);
+	if(PerSessionData != nullptr)
+	{
+		OnSessionEnded.Broadcast(*PerSessionData->SharedData);
+		PerSessionData->SharedData->OnEndSession(InSession);
+	}
+
+	PerSessionDataMap.Remove(&InSession);
+}
+
+void FNiagaraTimingViewExtender::Tick(UE::Insights::Timing::ITimingViewSession& InSession, const TraceServices::IAnalysisSession& InAnalysisSession)
+{
+	FPerSessionData* PerSessionData = PerSessionDataMap.Find(&InSession);
+	if(PerSessionData != nullptr)
+	{
+		PerSessionData->SharedData->Tick(InSession, InAnalysisSession);
+	}
+}
+
+void FNiagaraTimingViewExtender::ExtendFilterMenu(UE::Insights::Timing::ITimingViewSession& InSession, FMenuBuilder& InMenuBuilder)
+{
+	FPerSessionData* PerSessionData = PerSessionDataMap.Find(&InSession);
+	if (PerSessionData != nullptr)
+	{
+		PerSessionData->SharedData->ExtendFilterMenu(InMenuBuilder);
+	}
+}
+
+} //namespace UE::NiagaraInsights
+
+#undef LOCTEXT_NAMESPACE
diff --git a/Engine/Plugins/FX/NiagaraInsights/Source/NiagaraInsights/Private/NiagaraTimingViewExtender.h b/Engine/Plugins/FX/NiagaraInsights/Source/NiagaraInsights/Private/NiagaraTimingViewExtender.h
new file mode 100644
--- /dev/null
+++ b/Engine/Plugins/FX/NiagaraInsights/Source/NiagaraInsights/Private/NiagaraTimingViewExtender.h
@@ -0,0 +1,50 @@
+// Copyright Epic Games, Inc. All Rights Reserved.
+
+#pragma once
+
+#include "Insights/ITimingViewExtender.h"
+#include "NiagaraTimingViewSession.h"
+
+namespace UE::Insights::Timing { class ITimingViewSession; }
+namespace TraceServices { class IAnalysisSession; }
+
+class FMenuBuilder;
+
+namespace UE::NiagaraInsights
+{
+
+class FNiagaraTimingViewExtender : public UE::Insights::Timing::ITimingViewExtender
+{
+public:
+	DECLARE_MULTICAST_DELEGATE_OneParam(FOnSessionBegun, FNiagaraTimingViewSession& /*Session*/);
+	DECLARE_MULTICAST_DELEGATE_OneParam(FOnSessionEnded, FNiagaraTimingViewSession& /*Session*/);
+
+	//~ Begin UE::Insights::Timing::ITimingViewExtender impl
+	virtual void OnBeginSession(UE::Insights::Timing::ITimingViewSession& InSession) override;
+	virtual void OnEndSession(UE::Insights::Timing::ITimingViewSession& InSession) override;
+	virtual void Tick(UE::Insights::Timing::ITimingViewSession& InSession, const TraceServices::IAnalysisSession& InAnalysisSession) override;
+	virtual void ExtendFilterMenu(UE::Insights::Timing::ITimingViewSession& InSession, FMenuBuilder& InMenuBuilder) override;
+	//~ End UE::Insights::Timing::ITimingViewExtender impl
+
+	// Returns the first active Niagara timing-view session, or nullptr if none has started.
+	// The stats view uses this to bind to the range-selection delegate.
+	FNiagaraTimingViewSession* FindFirstActiveSession();
+
+	// Fired when a new timing-view session begins (after tracks are created).
+	// Observers can use this to late-bind to the session's OnRangeChanged delegate.
+	FOnSessionBegun OnSessionBegun;
+
+	// Fired when a timing-view session is about to be destroyed.
+	// Observers must release any raw pointers to the session before returning.
+	FOnSessionEnded OnSessionEnded;
+
+private:
+	struct FPerSessionData
+	{
+		TUniquePtr<FNiagaraTimingViewSession> SharedData;
+	};
+
+	TMap<UE::Insights::Timing::ITimingViewSession*, FPerSessionData> PerSessionDataMap;
+};
+
+} //namespace UE::NiagaraInsights
diff --git a/Engine/Plugins/FX/NiagaraInsights/Source/NiagaraInsights/Private/NiagaraTimingViewSession.cpp b/Engine/Plugins/FX/NiagaraInsights/Source/NiagaraInsights/Private/NiagaraTimingViewSession.cpp
new file mode 100644
--- /dev/null
+++ b/Engine/Plugins/FX/NiagaraInsights/Source/NiagaraInsights/Private/NiagaraTimingViewSession.cpp
@@ -0,0 +1,109 @@
+// Copyright Epic Games, Inc. All Rights Reserved.
+
+#include "NiagaraTimingViewSession.h"
+#include "NiagaraInstanceLifecycleTrack.h"
+#include "NiagaraPerformanceGraphTrack.h"
+#include "NiagaraProvider.h"
+
+#include "Framework/MultiBox/MultiBoxBuilder.h"
+#include "Insights/ITimingViewSession.h"
+
+#define LOCTEXT_NAMESPACE "NiagaraTimingViewSession"
+
+namespace UE::NiagaraInsights
+{
+
+void FNiagaraTimingViewSession::OnBeginSession(UE::Insights::Timing::ITimingViewSession& InTimingViewSession)
+{
+	TimingViewSession = &InTimingViewSession;
+
+	InstanceLifecycleTrack.Reset();
+	PerformanceGraphTrack.Reset();
+
+	InTimingViewSession.OnSelectionChanged().AddRaw(this, &FNiagaraTimingViewSession::HandleTimingViewSelectionChanged);
+}
+
+void FNiagaraTimingViewSession::OnEndSession(UE::Insights::Timing::ITimingViewSession& InTimingViewSession)
+{
+	InTimingViewSession.OnSelectionChanged().RemoveAll(this);
+
+	InstanceLifecycleTrack.Reset();
+	PerformanceGraphTrack.Reset();
+	TimingViewSession = nullptr;
+	AnalysisSession   = nullptr;
+}
+
+void FNiagaraTimingViewSession::HandleTimingViewSelectionChanged(UE::Insights::Timing::ETimeChangedFlags /*Flags*/, double StartTime, double EndTime)
+{
+	SelectionStartTime = StartTime;
+	SelectionEndTime   = EndTime;
+	OnRangeChanged.Broadcast(StartTime, EndTime);
+}
+
+void FNiagaraTimingViewSession::Tick(UE::Insights::Timing::ITimingViewSession& InTimingViewSession, const TraceServices::IAnalysisSession& InAnalysisSession)
+{
+	AnalysisSession = &InAnalysisSession;
+
+	const FNiagaraProvider* NiagaraProvider = InAnalysisSession.ReadProvider<FNiagaraProvider>(FNiagaraProvider::GetProviderName());
+
+	if (NiagaraProvider)
+	{
+		TraceServices::FAnalysisSessionReadScope SessionReadScope(*AnalysisSession);
+
+		//const FNiagaraTimelineDataPtr& TimelineData = NiagaraProvider->GetTimelineData();
+		//-TODO: Check we have data
+
+		if (!InstanceLifecycleTrack.IsValid())
+		{
+			PerformanceGraphTrack = MakeShared<FNiagaraPerformanceGraphTrack>(*this);
+			PerformanceGraphTrack->SetVisibilityFlag(bTrackVisible);
+			PerformanceGraphTrack->SetOrder((MAX_int32 / 2));
+			InTimingViewSession.AddScrollableTrack(PerformanceGraphTrack);
+
+			InstanceLifecycleTrack = MakeShared<FNiagaraInstanceLifecycleTrack>(*this);
+			InstanceLifecycleTrack->SetVisibilityFlag(bTrackVisible);
+			InstanceLifecycleTrack->SetOrder((MAX_int32 / 2) + 1);
+			InTimingViewSession.AddScrollableTrack(InstanceLifecycleTrack);
+
+			InTimingViewSession.InvalidateScrollableTracksOrder();
+		}
+	}
+}
+
+void FNiagaraTimingViewSession::ExtendFilterMenu(FMenuBuilder& InMenuBuilder)
+{
+	InMenuBuilder.BeginSection("NiagaraTracks", LOCTEXT("NiagaraHeader", "Niagara"));
+	{
+		InMenuBuilder.AddMenuEntry(
+			LOCTEXT("NiagaraTimingTracks", "Niagara Tracks"),
+			LOCTEXT("NiagaraTimingTracks_Tooltip", "Show/hide the Niagara track"),
+			FSlateIcon(),
+			FUIAction(
+				FExecuteAction::CreateRaw(this, &FNiagaraTimingViewSession::ToggleNiagaraTrack),
+				FCanExecuteAction(),
+				FIsActionChecked::CreateLambda([this]() { return bTrackVisible; })
+				),
+			NAME_None,
+			EUserInterfaceActionType::ToggleButton
+		);
+	}
+	InMenuBuilder.EndSection();
+}
+
+void FNiagaraTimingViewSession::ToggleNiagaraTrack()
+{
+	bTrackVisible = !bTrackVisible;
+
+	if (InstanceLifecycleTrack)
+	{
+		InstanceLifecycleTrack->SetVisibilityFlag(bTrackVisible);
+	}
+	if (PerformanceGraphTrack)
+	{
+		PerformanceGraphTrack->SetVisibilityFlag(bTrackVisible);
+	}
+}
+
+} //namespace UE::NiagaraInsights
+
+#undef LOCTEXT_NAMESPACE
diff --git a/Engine/Plugins/FX/NiagaraInsights/Source/NiagaraInsights/Private/NiagaraTimingViewSession.h b/Engine/Plugins/FX/NiagaraInsights/Source/NiagaraInsights/Private/NiagaraTimingViewSession.h
new file mode 100644
--- /dev/null
+++ b/Engine/Plugins/FX/NiagaraInsights/Source/NiagaraInsights/Private/NiagaraTimingViewSession.h
@@ -0,0 +1,56 @@
+// Copyright Epic Games, Inc. All Rights Reserved.
+
+#pragma once
+
+#include "Delegates/Delegate.h"
+#include "Templates/SharedPointer.h"
+
+namespace TraceServices { class IAnalysisSession; }
+namespace UE::Insights::Timing { class ITimingViewSession; }
+namespace UE::Insights::Timing { enum class ETimeChangedFlags : int32; }
+class FMenuBuilder;
+
+namespace UE::NiagaraInsights
+{
+
+class FNiagaraInstanceLifecycleTrack;
+class FNiagaraPerformanceGraphTrack;
+
+/** Fired when the user changes the time range selection in the Timing Profiler. */
+DECLARE_MULTICAST_DELEGATE_TwoParams(FOnNiagaraTimingRangeChanged, double /*StartTime*/, double /*EndTime*/);
+
+class FNiagaraTimingViewSession
+{
+public:
+	void OnBeginSession(UE::Insights::Timing::ITimingViewSession& InTimingViewSession);
+	void OnEndSession(UE::Insights::Timing::ITimingViewSession& InTimingViewSession);
+	void Tick(UE::Insights::Timing::ITimingViewSession& InTimingViewSession, const TraceServices::IAnalysisSession& InAnalysisSession);
+	void ExtendFilterMenu(FMenuBuilder& InMenuBuilder);
+
+	bool IsAnalysisSessionValid() const { return AnalysisSession != nullptr; }
+	const TraceServices::IAnalysisSession& GetAnalysisSession() const { check(AnalysisSession); return *AnalysisSession; }
+
+	double GetSelectionStartTime() const { return SelectionStartTime; }
+	double GetSelectionEndTime()   const { return SelectionEndTime; }
+
+	/** Broadcast when the timing-view selection range changes. */
+	FOnNiagaraTimingRangeChanged OnRangeChanged;
+
+	// Show / hide the track
+	void ToggleNiagaraTrack();
+
+private:
+	void HandleTimingViewSelectionChanged(UE::Insights::Timing::ETimeChangedFlags Flags, double StartTime, double EndTime);
+
+	const TraceServices::IAnalysisSession*		AnalysisSession = nullptr;
+	UE::Insights::Timing::ITimingViewSession*	TimingViewSession = nullptr;
+
+	TSharedPtr<FNiagaraInstanceLifecycleTrack>	InstanceLifecycleTrack;
+	TSharedPtr<FNiagaraPerformanceGraphTrack>	PerformanceGraphTrack;
+
+	bool   bTrackVisible      = true;
+	double SelectionStartTime = 0.0;
+	double SelectionEndTime   = 0.0;
+};
+
+} //namespace UE::NiagaraInsights
diff --git a/Engine/Plugins/FX/NiagaraInsights/Source/NiagaraInsights/Private/NiagaraTraceModule.cpp b/Engine/Plugins/FX/NiagaraInsights/Source/NiagaraInsights/Private/NiagaraTraceModule.cpp
new file mode 100644
--- /dev/null
+++ b/Engine/Plugins/FX/NiagaraInsights/Source/NiagaraInsights/Private/NiagaraTraceModule.cpp
@@ -0,0 +1,38 @@
+// Copyright Epic Games, Inc. All Rights Reserved.
+
+#include "NiagaraTraceModule.h"
+#include "NiagaraProvider.h"
+#include "NiagaraAnalyzer.h"
+
+namespace UE::NiagaraInsights
+{
+namespace Private
+{
+static FName ModuleName("TraceModule_Niagara");
+}
+
+void FNiagaraTraceModule::GetModuleInfo(TraceServices::FModuleInfo& OutModuleInfo)
+{
+	OutModuleInfo.Name = Private::ModuleName;
+	OutModuleInfo.DisplayName = TEXT("Niagara");
+}
+
+void FNiagaraTraceModule::OnAnalysisBegin(TraceServices::IAnalysisSession& InSession)
+{
+	TSharedPtr<FNiagaraProvider> Provider = MakeShared<FNiagaraProvider>(InSession);
+	InSession.AddProvider(FNiagaraProvider::GetProviderName(), Provider);
+
+	InSession.AddAnalyzer(new FNiagaraAnalyzer(InSession, *Provider));
+}
+
+void FNiagaraTraceModule::GetLoggers(TArray<const TCHAR *>& OutLoggers)
+{
+	OutLoggers.Add(TEXT("Niagara"));
+}
+
+void FNiagaraTraceModule::GenerateReports(const TraceServices::IAnalysisSession& Session, const TCHAR* CmdLine, const TCHAR* OutputDirectory)
+{
+
+}
+
+} //namespace UE::NiagaraInsights
diff --git a/Engine/Plugins/FX/NiagaraInsights/Source/NiagaraInsights/Private/NiagaraTraceModule.h b/Engine/Plugins/FX/NiagaraInsights/Source/NiagaraInsights/Private/NiagaraTraceModule.h
new file mode 100644
--- /dev/null
+++ b/Engine/Plugins/FX/NiagaraInsights/Source/NiagaraInsights/Private/NiagaraTraceModule.h
@@ -0,0 +1,22 @@
+// Copyright Epic Games, Inc. All Rights Reserved.
+
+#pragma once
+
+#include "TraceServices/ModuleService.h"
+
+namespace UE::NiagaraInsights
+{
+
+class FNiagaraTraceModule : public TraceServices::IModule
+{
+public:
+	//~ Begin TraceServices::IModule interface
+	virtual void GetModuleInfo(TraceServices::FModuleInfo& OutModuleInfo) override;
+	virtual void OnAnalysisBegin(TraceServices::IAnalysisSession& Session) override;
+	virtual void GetLoggers(TArray<const TCHAR *>& OutLoggers) override;
+	virtual void GenerateReports(const TraceServices::IAnalysisSession& Session, const TCHAR* CmdLine, const TCHAR* OutputDirectory) override;
+	virtual const TCHAR* GetCommandLineArgument() override { return TEXT("niagaratrace"); }
+	//~ End TraceServices::IModule interface
+};
+
+} //namespace UE::NiagaraInsights
diff --git a/Engine/Plugins/FX/NiagaraInsights/Source/NiagaraInsights/Private/Widgets/SNiagaraRangeStatsView.cpp b/Engine/Plugins/FX/NiagaraInsights/Source/NiagaraInsights/Private/Widgets/SNiagaraRangeStatsView.cpp
new file mode 100644
--- /dev/null
+++ b/Engine/Plugins/FX/NiagaraInsights/Source/NiagaraInsights/Private/Widgets/SNiagaraRangeStatsView.cpp
@@ -0,0 +1,583 @@
+// Copyright Epic Games, Inc. All Rights Reserved.
+
+#include "SNiagaraRangeStatsView.h"
+
+#include "NiagaraInsightsCommon.h"
+#include "NiagaraInsightsModule.h"
+#include "NiagaraProvider.h"
+#include "NiagaraTimingViewExtender.h"
+#include "NiagaraTimingViewSession.h"
+
+#include "Insights/IUnrealInsightsModule.h"
+#include "Modules/ModuleManager.h"
+#include "TraceServices/Model/AnalysisSession.h"
+
+#include "Widgets/Layout/SExpandableArea.h"
+#include "Widgets/Layout/SBox.h"
+#include "Widgets/Layout/SScrollBox.h"
+#include "Widgets/Layout/SSeparator.h"
+#include "Widgets/SBoxPanel.h"
+#include "Widgets/Text/STextBlock.h"
+#include "Widgets/Views/SHeaderRow.h"
+#include "Widgets/Views/STableRow.h"
+
+#define LOCTEXT_NAMESPACE "SNiagaraRangeStatsView"
+
+namespace UE::NiagaraInsights
+{
+
+// ---- Column IDs ----------------------------------------------------------------
+
+static const FName ColId_Rank(TEXT("Rank"));
+static const FName ColId_System(TEXT("System"));
+static const FName ColId_Primary(TEXT("Primary"));
+static const FName ColId_Secondary(TEXT("Secondary"));
+
+// ---- STableRow for a stats row -------------------------------------------------
+
+class SNiagaraStatsTableRow : public SMultiColumnTableRow<FNiagaraStatsRowPtr>
+{
+public:
+	SLATE_BEGIN_ARGS(SNiagaraStatsTableRow) {}
+		SLATE_ARGUMENT(FNiagaraStatsRowPtr, Row)
+		SLATE_ARGUMENT(FText, PrimaryColumnLabel)
+		SLATE_ARGUMENT(FText, SecondaryColumnLabel)
+	SLATE_END_ARGS()
+
+	void Construct(const FArguments& InArgs, const TSharedRef<STableViewBase>& InOwnerTable)
+	{
+		Row                  = InArgs._Row;
+		PrimaryColumnLabel   = InArgs._PrimaryColumnLabel;
+		SecondaryColumnLabel = InArgs._SecondaryColumnLabel;
+		SMultiColumnTableRow::Construct(SMultiColumnTableRow::FArguments(), InOwnerTable);
+	}
+
+	virtual TSharedRef<SWidget> GenerateWidgetForColumn(const FName& ColumnId) override
+	{
+		if (!Row.IsValid())
+		{
+			return SNullWidget::NullWidget;
+		}
+
+		const FMargin CellPad(4.f, 2.f);
+
+		if (ColumnId == ColId_Rank)
+		{
+			return SNew(SBox).Padding(CellPad).HAlign(HAlign_Center)
+				[
+					SNew(STextBlock)
+					.Text(FText::AsNumber(Row->Rank))
+					.ColorAndOpacity(FSlateColor::UseSubduedForeground())
+				];
+		}
+		if (ColumnId == ColId_System)
+		{
+			return SNew(SBox).Padding(CellPad)
+				[
+					SNew(STextBlock)
+					.Text(FText::FromString(Row->SystemName))
+					.ToolTipText(FText::FromString(Row->SystemName))
+				];
+		}
+		if (ColumnId == ColId_Primary)
+		{
+			const FString Formatted = FString::Printf(TEXT("%.3f %s"), Row->PrimaryValue, *Row->PrimaryUnit);
+			return SNew(SBox).Padding(CellPad).HAlign(HAlign_Right)
+				[
+					SNew(STextBlock).Text(FText::FromString(Formatted))
+				];
+		}
+		if (ColumnId == ColId_Secondary)
+		{
+			const FString Formatted = FString::Printf(TEXT("%.3f %s"), Row->SecondaryValue, *Row->SecondaryUnit);
+			return SNew(SBox).Padding(CellPad).HAlign(HAlign_Right)
+				[
+					SNew(STextBlock)
+					.Text(FText::FromString(Formatted))
+					.ColorAndOpacity(FSlateColor::UseSubduedForeground())
+				];
+		}
+		return SNullWidget::NullWidget;
+	}
+
+private:
+	FNiagaraStatsRowPtr Row;
+	FText PrimaryColumnLabel;
+	FText SecondaryColumnLabel;
+};
+
+// ---- SNiagaraRangeStatsView ----------------------------------------------------
+
+void SNiagaraRangeStatsView::Construct(const FArguments& InArgs)
+{
+	RangeText      = LOCTEXT("NoRange",    "No range selected");
+	FrameCountText = LOCTEXT("NoSession",  "No active Niagara session");
+
+	// Seed empty rows so the list views have something to bind against.
+	auto MakeEmpty = [](TArray<FNiagaraStatsRowPtr>& Rows)
+	{
+		for (int32 i = 1; i <= 5; ++i)
+		{
+			auto Row = MakeShared<FNiagaraStatsRow>();
+			Row->Rank = i;
+			Row->SystemName = TEXT("—");
+			Rows.Add(Row);
+		}
+	};
+	MakeEmpty(Rows_InstanceCount);
+	MakeEmpty(Rows_GTCostTotal);
+	MakeEmpty(Rows_GTCostPerInst);
+	MakeEmpty(Rows_RTCostTotal);
+	MakeEmpty(Rows_RTCostPerInst);
+	MakeEmpty(Rows_GpuCostTotal);
+
+	ChildSlot
+	[
+		SNew(SVerticalBox)
+
+		// ---- Range info header ----
+		+ SVerticalBox::Slot()
+		.AutoHeight()
+		.Padding(4.f, 4.f, 4.f, 2.f)
+		[
+			MakeRangeHeaderBar()
+		]
+
+		+ SVerticalBox::Slot()
+		.AutoHeight()
+		[
+			SNew(SSeparator)
+		]
+
+		// ---- Six top-5 sections in a scroll box ----
+		+ SVerticalBox::Slot()
+		.FillHeight(1.f)
+		[
+			SNew(SScrollBox)
+
+			+ SScrollBox::Slot()
+			[
+				MakeStatsSection(
+					LOCTEXT("Sec_InstanceCount", "Top 5 by Instance Count"),
+					Common::Color_InstanceCount,
+					LOCTEXT("Col_AvgInstances", "Avg Instances (GT)"),
+					LOCTEXT("Col_TotalGTms",    "Total GT (ms)"),
+					Rows_InstanceCount)
+			]
+
+			+ SScrollBox::Slot()
+			.Padding(0.f, 4.f, 0.f, 0.f)
+			[
+				MakeStatsSection(
+					LOCTEXT("Sec_GTCostTotal", "Top 5 by Game Thread Cost"),
+					Common::Color_GTCostTotal,
+					LOCTEXT("Col_TotalGTms2",   "Total GT (ms)"),
+					LOCTEXT("Col_AvgInst2",     "Avg Instances"),
+					Rows_GTCostTotal)
+			]
+
+			+ SScrollBox::Slot()
+			.Padding(0.f, 4.f, 0.f, 0.f)
+			[
+				MakeStatsSection(
+					LOCTEXT("Sec_GTCostPerInst", "Top 5 by GT Cost / Instance"),
+					Common::Color_GTCostPerInst,
+					LOCTEXT("Col_GTmsPerInst",  "GT ms / Instance"),
+					LOCTEXT("Col_TotalGT3",     "Total GT (ms)"),
+					Rows_GTCostPerInst)
+			]
+
+			+ SScrollBox::Slot()
+			.Padding(0.f, 4.f, 0.f, 0.f)
+			[
+				MakeStatsSection(
+					LOCTEXT("Sec_RTCostTotal", "Top 5 by Render Thread Cost"),
+					Common::Color_RTCostTotal,
+					LOCTEXT("Col_TotalRTms",    "Total RT (ms)"),
+					LOCTEXT("Col_AvgInst3",     "Avg Instances"),
+					Rows_RTCostTotal)
+			]
+
+			+ SScrollBox::Slot()
+			.Padding(0.f, 4.f, 0.f, 0.f)
+			[
+				MakeStatsSection(
+					LOCTEXT("Sec_RTCostPerInst", "Top 5 by RT Cost / Instance"),
+					Common::Color_RTCostPerInst,
+					LOCTEXT("Col_RTmsPerInst",  "RT ms / Instance"),
+					LOCTEXT("Col_TotalRT3",     "Total RT (ms)"),
+					Rows_RTCostPerInst)
+			]
+
+			+ SScrollBox::Slot()
+			.Padding(0.f, 4.f, 0.f, 4.f)
+			[
+				MakeStatsSection(
+					LOCTEXT("Sec_GpuCostTotal", "Top 5 by GPU Cost"),
+					Common::Color_GpuCostTotal,
+					LOCTEXT("Col_TotalGPUms",   "Total GPU (ms)"),
+					LOCTEXT("Col_GpuInstances", "GPU Instances"),
+					Rows_GpuCostTotal)
+			]
+		]
+	];
+
+	// Bind to the extender if a session is already active.
+	TryBindToActiveSession();
+
+	// Also subscribe to future sessions.
+	if (FModuleManager::Get().IsModuleLoaded("NiagaraInsights"))
+	{
+		FNiagaraInsightsModule& Module = FNiagaraInsightsModule::Get();
+		Module.GetTimingViewExtender().OnSessionBegun.AddSP(this, &SNiagaraRangeStatsView::OnNiagaraTimingSessionBegun);
+		Module.GetTimingViewExtender().OnSessionEnded.AddSP(this, &SNiagaraRangeStatsView::OnNiagaraTimingSessionEnded);
+	}
+}
+
+SNiagaraRangeStatsView::~SNiagaraRangeStatsView()
+{
+	UnbindCurrentSession();
+
+	if (FModuleManager::Get().IsModuleLoaded("NiagaraInsights"))
+	{
+		FNiagaraInsightsModule& Module = FNiagaraInsightsModule::Get();
+		Module.GetTimingViewExtender().OnSessionBegun.RemoveAll(this);
+		Module.GetTimingViewExtender().OnSessionEnded.RemoveAll(this);
+	}
+}
+
+// ---- Session binding -----------------------------------------------------------
+
+void SNiagaraRangeStatsView::TryBindToActiveSession()
+{
+	if (!FModuleManager::Get().IsModuleLoaded("NiagaraInsights"))
+	{
+		return;
+	}
+
+	FNiagaraTimingViewSession* Session =
+		FNiagaraInsightsModule::Get().GetTimingViewExtender().FindFirstActiveSession();
+
+	if (Session && Session != BoundSession)
+	{
+		UnbindCurrentSession();
+
+		BoundSession     = Session;
+		RangeChangedHandle = Session->OnRangeChanged.AddSP(
+			this, &SNiagaraRangeStatsView::OnTimingRangeChanged);
+	}
+}
+
+void SNiagaraRangeStatsView::OnNiagaraTimingSessionEnded(FNiagaraTimingViewSession& Session)
+{
+	if (&Session == BoundSession)
+	{
+		UnbindCurrentSession();
+	}
+}
+
+void SNiagaraRangeStatsView::OnNiagaraTimingSessionBegun(FNiagaraTimingViewSession& Session)
+{
+	if (&Session == BoundSession)
+	{
+		return;
+	}
+	UnbindCurrentSession();
+
+	BoundSession       = &Session;
+	RangeChangedHandle = Session.OnRangeChanged.AddSP(
+		this, &SNiagaraRangeStatsView::OnTimingRangeChanged);
+}
+
+void SNiagaraRangeStatsView::UnbindCurrentSession()
+{
+	if (BoundSession)
+	{
+		BoundSession->OnRangeChanged.Remove(RangeChangedHandle);
+		BoundSession = nullptr;
+	}
+	RangeChangedHandle.Reset();
+}
+
+void SNiagaraRangeStatsView::OnTimingRangeChanged(double StartTime, double EndTime)
+{
+	RebuildStats(StartTime, EndTime);
+}
+
+// ---- Data aggregation ----------------------------------------------------------
+
+void SNiagaraRangeStatsView::RebuildStats(double StartTime, double EndTime)
+{
+	CurrentStartTime = StartTime;
+	CurrentEndTime   = EndTime;
+
+	const double DurationS = EndTime - StartTime;
+
+	// Update header labels.
+	if (RangeTextWidget.IsValid())
+	{
+		RangeTextWidget->SetText(FText::FromString(
+			FString::Printf(TEXT("Range: %.4f s  →  %.4f s   (%.4f s)"),
+				StartTime, EndTime, DurationS)));
+	}
+
+	// Fetch the analysis session.
+	if (!FModuleManager::Get().IsModuleLoaded("TraceInsights"))
+	{
+		return;
+	}
+	IUnrealInsightsModule& InsightsModule =
+		FModuleManager::LoadModuleChecked<IUnrealInsightsModule>("TraceInsights");
+	TSharedPtr<const TraceServices::IAnalysisSession> AnalysisSession =
+		InsightsModule.GetAnalysisSession();
+	if (!AnalysisSession.IsValid())
+	{
+		return;
+	}
+
+	TraceServices::FAnalysisSessionReadScope ReadScope(*AnalysisSession);
+
+	const FNiagaraProvider* Provider =
+		AnalysisSession->ReadProvider<FNiagaraProvider>(FNiagaraProvider::GetProviderName());
+	if (!Provider)
+	{
+		return;
+	}
+
+	// Accumulate per-system data across the range.
+	TMap<FString, FAggregated> AggMap;
+
+	int32 GTFrameCount = 0;
+	Provider->EnumeratePerformance_GT(StartTime, EndTime,
+		[&](const FSystemPerformanceFrame_GT& Frame)
+		{
+			++GTFrameCount;
+			for (const FSystemPerformanceFrame_GT::FSystemData& Sys : Frame.SystemData)
+			{
+				FAggregated& A = AggMap.FindOrAdd(Sys.SystemName);
+				A.SystemName = Sys.SystemName;
+
+				A.GTSamples++;
+				A.GTTotalInstances += (double)Sys.NumInstances;
+				A.GTTotalMs +=
+					(Sys.TickGameThreadSeconds +
+					 Sys.TickConcurrentSeconds +
+					 Sys.FinalizeSeconds +
+					 Sys.EndOfFrameSeconds +
+					 Sys.ActivationSeconds +
+					 Sys.WaitSeconds) * 1000.0;
+			}
+		});
+
+	Provider->EnumeratePerformance_RT(StartTime, EndTime,
+		[&](const FSystemPerformanceFrame_RT& Frame)
+		{
+			for (const FSystemPerformanceFrame_RT::FSystemData& Sys : Frame.SystemData)
+			{
+				FAggregated& A = AggMap.FindOrAdd(Sys.SystemName);
+				A.SystemName = Sys.SystemName;
+
+				A.RTSamples++;
+				A.RTTotalInstances += (double)Sys.NumInstances;
+				A.RTTotalMs        +=
+					(Sys.RenderUpdateSeconds + Sys.GetDynamicMeshElementsSeconds) * 1000.0;
+				A.GpuTotalMs  += (double)Sys.GpuTotalMicroseconds / 1000.0;
+				A.GpuTotalInst += (double)Sys.GpuNumInstances;
+			}
+		});
+
+	if (FrameCountWidget.IsValid())
+	{
+		FrameCountWidget->SetText(FText::FromString(
+			FString::Printf(TEXT("%d GT frames  |  %d systems"),
+				GTFrameCount, AggMap.Num())));
+	}
+
+	// Convert map to array for sorting.
+	TArray<FAggregated> AggArray;
+	AggArray.Reserve(AggMap.Num());
+	for (auto& KV : AggMap)
+	{
+		AggArray.Add(MoveTemp(KV.Value));
+	}
+
+	// Build six top-5 lists.
+	Rows_InstanceCount = BuildTop5(AggArray,
+		[](const FAggregated& A) { return (float)A.AvgGTInstances(); },  TEXT("inst"),
+		[](const FAggregated& A) { return (float)A.GTTotalMs; },         TEXT("ms"));
+
+	Rows_GTCostTotal = BuildTop5(AggArray,
+		[](const FAggregated& A) { return (float)A.GTTotalMs; },         TEXT("ms"),
+		[](const FAggregated& A) { return (float)A.AvgGTInstances(); },  TEXT("inst"));
+
+	Rows_GTCostPerInst = BuildTop5(AggArray,
+		[](const FAggregated& A) { return (float)A.GTCostPerInst(); },   TEXT("ms/inst"),
+		[](const FAggregated& A) { return (float)A.GTTotalMs; },         TEXT("ms"));
+
+	Rows_RTCostTotal = BuildTop5(AggArray,
+		[](const FAggregated& A) { return (float)A.RTTotalMs; },         TEXT("ms"),
+		[](const FAggregated& A) { return (float)A.RTTotalInstances / FMath::Max(1.0, (double)A.RTSamples); }, TEXT("inst"));
+
+	Rows_RTCostPerInst = BuildTop5(AggArray,
+		[](const FAggregated& A) { return (float)A.RTCostPerInst(); },   TEXT("ms/inst"),
+		[](const FAggregated& A) { return (float)A.RTTotalMs; },         TEXT("ms"));
+
+	Rows_GpuCostTotal = BuildTop5(AggArray,
+		[](const FAggregated& A) { return (float)A.GpuTotalMs; },        TEXT("ms"),
+		[](const FAggregated& A) { return (float)A.GpuCostPerInst(); },  TEXT("ms/inst"));
+
+	// Refresh all list views.
+	auto Refresh = [](TSharedPtr<SListView<FNiagaraStatsRowPtr>>& LV)
+	{
+		if (LV.IsValid()) { LV->RequestListRefresh(); }
+	};
+	Refresh(ListView_InstanceCount);
+	Refresh(ListView_GTCostTotal);
+	Refresh(ListView_GTCostPerInst);
+	Refresh(ListView_RTCostTotal);
+	Refresh(ListView_RTCostPerInst);
+	Refresh(ListView_GpuCostTotal);
+}
+
+// ---- Build top-5 helper --------------------------------------------------------
+
+TArray<FNiagaraStatsRowPtr> SNiagaraRangeStatsView::BuildTop5(
+	const TArray<FAggregated>& Data,
+	TFunctionRef<float(const FAggregated&)> PrimaryGetter,
+	const FString& PrimaryUnit,
+	TFunctionRef<float(const FAggregated&)> SecondaryGetter,
+	const FString& SecondaryUnit)
+{
+	// Sort descending by primary metric.
+	TArray<const FAggregated*> Sorted;
+	Sorted.Reserve(Data.Num());
+	for (const FAggregated& A : Data)
+	{
+		Sorted.Add(&A);
+	}
+	Sorted.Sort([&PrimaryGetter](const FAggregated& A, const FAggregated& B)
+	{
+		return PrimaryGetter(A) > PrimaryGetter(B);
+	});
+
+	TArray<FNiagaraStatsRowPtr> Rows;
+	const int32 Count = FMath::Min(Sorted.Num(), 5);
+	for (int32 i = 0; i < Count; ++i)
+	{
+		auto Row = MakeShared<FNiagaraStatsRow>();
+		Row->Rank           = i + 1;
+		Row->SystemName     = Sorted[i]->SystemName;
+		Row->PrimaryValue   = PrimaryGetter(*Sorted[i]);
+		Row->PrimaryUnit    = PrimaryUnit;
+		Row->SecondaryValue = SecondaryGetter(*Sorted[i]);
+		Row->SecondaryUnit  = SecondaryUnit;
+		Rows.Add(MoveTemp(Row));
+	}
+
+	// Pad to 5 rows so the table always looks uniform.
+	while (Rows.Num() < 5)
+	{
+		auto Row = MakeShared<FNiagaraStatsRow>();
+		Row->Rank = Rows.Num() + 1;
+		Row->SystemName = TEXT("—");
+		Rows.Add(MoveTemp(Row));
+	}
+
+	return Rows;
+}
+
+// ---- UI construction helpers ---------------------------------------------------
+
+TSharedRef<SWidget> SNiagaraRangeStatsView::MakeRangeHeaderBar()
+{
+	return SNew(SHorizontalBox)
+
+		+ SHorizontalBox::Slot()
+		.AutoWidth()
+		.Padding(0.f, 0.f, 16.f, 0.f)
+		[
+			SAssignNew(RangeTextWidget, STextBlock)
+			.Text(RangeText)
+		]
+
+		+ SHorizontalBox::Slot()
+		.AutoWidth()
+		[
+			SAssignNew(FrameCountWidget, STextBlock)
+			.Text(FrameCountText)
+			.ColorAndOpacity(FSlateColor::UseSubduedForeground())
+		];
+}
+
+TSharedRef<SWidget> SNiagaraRangeStatsView::MakeStatsSection(
+	const FText&                 SectionTitle,
+	const FLinearColor&          HeaderColor,
+	const FText&                 PrimaryColumnLabel,
+	const FText&                 SecondaryColumnLabel,
+	TArray<FNiagaraStatsRowPtr>& Rows)
+{
+	// Create a list view pointer reference to store.
+	TSharedPtr<SListView<FNiagaraStatsRowPtr>> LocalListView;
+
+	TSharedRef<SHeaderRow> Header =
+		SNew(SHeaderRow)
+		+ SHeaderRow::Column(ColId_Rank)
+			.DefaultLabel(LOCTEXT("Col_Rank", "#"))
+			.FixedWidth(28.f)
+		+ SHeaderRow::Column(ColId_System)
+			.DefaultLabel(LOCTEXT("Col_System", "System"))
+			.FillWidth(1.f)
+		+ SHeaderRow::Column(ColId_Primary)
+			.DefaultLabel(PrimaryColumnLabel)
+			.FixedWidth(130.f)
+		+ SHeaderRow::Column(ColId_Secondary)
+			.DefaultLabel(SecondaryColumnLabel)
+			.FixedWidth(130.f);
+
+	TSharedRef<SListView<FNiagaraStatsRowPtr>> ListView =
+		SNew(SListView<FNiagaraStatsRowPtr>)
+		.ListItemsSource(&Rows)
+		.HeaderRow(Header)
+		.OnGenerateRow_Lambda(
+			[this, PrimaryColumnLabel, SecondaryColumnLabel]
+			(FNiagaraStatsRowPtr Item, const TSharedRef<STableViewBase>& OwnerTable)
+			{
+				return SNew(SNiagaraStatsTableRow, OwnerTable)
+					.Row(Item)
+					.PrimaryColumnLabel(PrimaryColumnLabel)
+					.SecondaryColumnLabel(SecondaryColumnLabel);
+			});
+
+	LocalListView = ListView;
+
+	// Store back into the appropriate member pointer based on which row array this is.
+	// We use pointer comparison to identify the section.
+	if      (&Rows == &Rows_InstanceCount) { ListView_InstanceCount = LocalListView; }
+	else if (&Rows == &Rows_GTCostTotal)   { ListView_GTCostTotal   = LocalListView; }
+	else if (&Rows == &Rows_GTCostPerInst) { ListView_GTCostPerInst = LocalListView; }
+	else if (&Rows == &Rows_RTCostTotal)   { ListView_RTCostTotal   = LocalListView; }
+	else if (&Rows == &Rows_RTCostPerInst) { ListView_RTCostPerInst = LocalListView; }
+	else if (&Rows == &Rows_GpuCostTotal)  { ListView_GpuCostTotal  = LocalListView; }
+
+	return SNew(SExpandableArea)
+		.InitiallyCollapsed(false)
+		.HeaderContent()
+		[
+			SNew(SBorder)
+			.BorderBackgroundColor(HeaderColor)
+			.Padding(FMargin(8.f, 4.f))
+			[
+				SNew(STextBlock)
+				.Text(SectionTitle)
+				.ColorAndOpacity(FLinearColor::White)
+			]
+		]
+		.BodyContent()
+		[
+			ListView
+		];
+}
+
+} // namespace UE::NiagaraInsights
+
+#undef LOCTEXT_NAMESPACE
diff --git a/Engine/Plugins/FX/NiagaraInsights/Source/NiagaraInsights/Private/Widgets/SNiagaraRangeStatsView.h b/Engine/Plugins/FX/NiagaraInsights/Source/NiagaraInsights/Private/Widgets/SNiagaraRangeStatsView.h
new file mode 100644
--- /dev/null
+++ b/Engine/Plugins/FX/NiagaraInsights/Source/NiagaraInsights/Private/Widgets/SNiagaraRangeStatsView.h
@@ -0,0 +1,143 @@
+// Copyright Epic Games, Inc. All Rights Reserved.
+
+#pragma once
+
+#include "CoreMinimal.h"
+#include "Widgets/SCompoundWidget.h"
+#include "Widgets/Views/SListView.h"
+
+namespace UE::NiagaraInsights
+{
+
+class FNiagaraTimingViewSession;
+
+/** One row in a top-5 table. */
+struct FNiagaraStatsRow
+{
+	int32   Rank           = 0;
+	FString SystemName;
+	float   PrimaryValue   = 0.f;   // The metric used for ranking (ms, count, …)
+	FString PrimaryUnit;            // e.g. "ms", "instances", "ms/inst"
+	float   SecondaryValue = 0.f;   // Supplementary context (instances, total ms, …)
+	FString SecondaryUnit;
+};
+
+using FNiagaraStatsRowPtr = TSharedPtr<FNiagaraStatsRow>;
+
+/**
+ * Displays top-5 Niagara systems across six metric categories for the
+ * time range selected in the Timing Profiler.
+ *
+ * Categories:
+ *   1. Instance Count           — avg active instances (GT)
+ *   2. Game Thread Cost         — total GT cost (ms)
+ *   3. GT Cost / Instance       — avg GT ms per live instance
+ *   4. Render Thread Cost       — total RT cost (ms)
+ *   5. RT Cost / Instance       — avg RT ms per live instance
+ *   6. GPU Cost                 — total GPU cost (ms)
+ */
+class SNiagaraRangeStatsView : public SCompoundWidget
+{
+public:
+	SLATE_BEGIN_ARGS(SNiagaraRangeStatsView) {}
+	SLATE_END_ARGS()
+
+	void Construct(const FArguments& InArgs);
+	virtual ~SNiagaraRangeStatsView();
+
+private:
+	// ---- Data ----------------------------------------------------------------
+
+	/** Aggregated per-system data accumulated while scanning a range. */
+	struct FAggregated
+	{
+		FString SystemName;
+
+		// Game thread
+		uint64  GTSamples        = 0;
+		double  GTTotalInstances = 0.0;   // sum of NumInstances across GT frames
+		double  GTTotalMs        = 0.0;   // total GT wall-clock ms
+
+		// Render thread
+		uint64  RTSamples        = 0;
+		double  RTTotalInstances = 0.0;
+		double  RTTotalMs        = 0.0;
+
+		// GPU
+		double  GpuTotalMs       = 0.0;
+		double  GpuTotalInst     = 0.0;
+
+		// Derived
+		double AvgGTInstances()   const { return GTSamples  > 0 ? GTTotalInstances / GTSamples  : 0.0; }
+		double GTCostPerInst()    const { return GTTotalInstances > 0 ? GTTotalMs / GTTotalInstances : 0.0; }
+		double RTCostPerInst()    const { return RTTotalInstances > 0 ? RTTotalMs / RTTotalInstances : 0.0; }
+		double GpuCostPerInst()   const { return GpuTotalInst     > 0 ? GpuTotalMs / GpuTotalInst  : 0.0; }
+	};
+
+	void RebuildStats(double StartTime, double EndTime);
+
+	static TArray<FNiagaraStatsRowPtr> BuildTop5(
+		const TArray<FAggregated>& Data,
+		TFunctionRef<float(const FAggregated&)> PrimaryGetter,
+		const FString& PrimaryUnit,
+		TFunctionRef<float(const FAggregated&)> SecondaryGetter,
+		const FString& SecondaryUnit);
+
+	// ---- UI helpers ----------------------------------------------------------
+
+	TSharedRef<SWidget> MakeRangeHeaderBar();
+
+	/** Build one collapsible top-5 section. */
+	TSharedRef<SWidget> MakeStatsSection(
+		const FText&                       SectionTitle,
+		const FLinearColor&                HeaderColor,
+		const FText&                       PrimaryColumnLabel,
+		const FText&                       SecondaryColumnLabel,
+		TArray<FNiagaraStatsRowPtr>&       Rows);
+
+	TSharedRef<ITableRow> OnGenerateRow(
+		FNiagaraStatsRowPtr Item,
+		const TSharedRef<STableViewBase>& OwnerTable,
+		const FText& PrimaryColumnLabel,
+		const FText& SecondaryColumnLabel);
+
+	// ---- Session binding -------------------------------------------------------
+
+	void TryBindToActiveSession();
+	void OnNiagaraTimingSessionBegun(FNiagaraTimingViewSession& Session);
+	void OnNiagaraTimingSessionEnded(FNiagaraTimingViewSession& Session);
+	void OnTimingRangeChanged(double StartTime, double EndTime);
+	void UnbindCurrentSession();
+
+	FNiagaraTimingViewSession* BoundSession = nullptr;
+	FDelegateHandle            RangeChangedHandle;
+
+	// ---- State ---------------------------------------------------------------
+
+	FText    RangeText;
+	FText    FrameCountText;
+	double   CurrentStartTime = 0.0;
+	double   CurrentEndTime   = 0.0;
+
+	// Six top-5 row arrays — rebuilt whenever the range changes.
+	TArray<FNiagaraStatsRowPtr> Rows_InstanceCount;
+	TArray<FNiagaraStatsRowPtr> Rows_GTCostTotal;
+	TArray<FNiagaraStatsRowPtr> Rows_GTCostPerInst;
+	TArray<FNiagaraStatsRowPtr> Rows_RTCostTotal;
+	TArray<FNiagaraStatsRowPtr> Rows_RTCostPerInst;
+	TArray<FNiagaraStatsRowPtr> Rows_GpuCostTotal;
+
+	// List-view widgets so we can call RequestListRefresh().
+	TSharedPtr<SListView<FNiagaraStatsRowPtr>> ListView_InstanceCount;
+	TSharedPtr<SListView<FNiagaraStatsRowPtr>> ListView_GTCostTotal;
+	TSharedPtr<SListView<FNiagaraStatsRowPtr>> ListView_GTCostPerInst;
+	TSharedPtr<SListView<FNiagaraStatsRowPtr>> ListView_RTCostTotal;
+	TSharedPtr<SListView<FNiagaraStatsRowPtr>> ListView_RTCostPerInst;
+	TSharedPtr<SListView<FNiagaraStatsRowPtr>> ListView_GpuCostTotal;
+
+	// Range header text widgets so we can update them in-place.
+	TSharedPtr<STextBlock> RangeTextWidget;
+	TSharedPtr<STextBlock> FrameCountWidget;
+};
+
+} // namespace UE::NiagaraInsights
diff --git a/Engine/Source/Runtime/Engine/Classes/Particles/ParticlePerfStats.h b/Engine/Source/Runtime/Engine/Classes/Particles/ParticlePerfStats.h
--- a/Engine/Source/Runtime/Engine/Classes/Particles/ParticlePerfStats.h
+++ b/Engine/Source/Runtime/Engine/Classes/Particles/ParticlePerfStats.h
@@ -48,6 +48,7 @@ struct FParticlePerfStats_GT
 	TAtomic<uint64> EndOfFrameCycles;
 	TAtomic<uint64> ActivationCycles;
 	uint64 WaitCycles;
+	uint64 MemoryBytes;
 
 	FParticlePerfStats_GT() { Reset(); }
 	
@@ -60,6 +61,7 @@ struct FParticlePerfStats_GT
 		EndOfFrameCycles = Other.EndOfFrameCycles.Load();
 		ActivationCycles = Other.ActivationCycles.Load();
 		WaitCycles = Other.WaitCycles;
+		MemoryBytes = Other.MemoryBytes;
 	}
 
 	FParticlePerfStats_GT& operator=(const FParticlePerfStats_GT& Other)
@@ -71,6 +73,7 @@ struct FParticlePerfStats_GT
 		EndOfFrameCycles = Other.EndOfFrameCycles.Load();
 		ActivationCycles = Other.ActivationCycles.Load();
 		WaitCycles = Other.WaitCycles;
+		MemoryBytes = Other.MemoryBytes;
 		return *this;
 	}
 
@@ -96,6 +99,7 @@ struct FParticlePerfStats_GT
 		EndOfFrameCycles += Other.EndOfFrameCycles.Load();
 		ActivationCycles += Other.ActivationCycles.Load();
 		WaitCycles += Other.WaitCycles;
+		MemoryBytes += Other.MemoryBytes;
 		return *this;
 	}
 
@@ -108,6 +112,7 @@ struct FParticlePerfStats_GT
 		EndOfFrameCycles = 0;
 		ActivationCycles = 0;
 		WaitCycles = 0;
+		MemoryBytes = 0;
 	}
 
 	inline uint64 GetTotalCycles_GTOnly() const { return TickGameThreadCycles + FinalizeCycles + ActivationCycles + WaitCycles; }
@@ -285,12 +290,14 @@ struct FParticlePerfStats
 	FName CSVStat_Culled = NAME_None;
 	FName CSVStat_MemoryKB = NAME_None;
 
-	TOptional<uint64> CSVMemoryKB_Asset;
-
 	void PopulateStatNames(const FName InName);
 	void ResetStatNames();
 #endif
 
+#if WITH_PER_SYSTEM_PARTICLE_PERF_STATS
+	TOptional<uint64> MemoryKB_Asset;
+#endif
+
 private:
 	static ENGINE_API FParticlePerfStats* GetWorldPerfStats(const UWorld* World);
 	static ENGINE_API FParticlePerfStats* GetSystemPerfStats(const UFXSystemAsset* FXAsset);
diff --git a/Engine/Source/Runtime/Engine/Classes/Particles/ParticlePerfStatsManager.h b/Engine/Source/Runtime/Engine/Classes/Particles/ParticlePerfStatsManager.h
--- a/Engine/Source/Runtime/Engine/Classes/Particles/ParticlePerfStatsManager.h
+++ b/Engine/Source/Runtime/Engine/Classes/Particles/ParticlePerfStatsManager.h
@@ -199,7 +199,7 @@ public:
 	/** Called when a component has been freed and is no longer tracked by the stats. */
 	virtual void OnRemoveComponent(const TWeakObjectPtr<const UFXSystemComponent>& Component) {}
 
-
+	virtual bool NeedsMemoryStats() const { return false; }
 	virtual bool NeedsWorldStats() const = 0;
 	virtual bool NeedsSystemStats() const = 0;
 	virtual bool NeedsComponentStats() const = 0;
@@ -333,6 +333,7 @@ public:
 	ENGINE_API virtual void OnRemoveComponent(const TWeakObjectPtr<const UFXSystemComponent>& Component)override;
 	#endif
 
+	virtual bool NeedsMemoryStats() const override { return true; }
 	virtual bool NeedsWorldStats() const override { return bGatherWorldStats; }
 	virtual bool NeedsSystemStats() const override { return bGatherSystemStats; }
 	virtual bool NeedsComponentStats() const override { return bGatherComponentStats; }
diff --git a/Engine/Source/Runtime/Engine/Private/Particles/ParticlePerfStats.cpp b/Engine/Source/Runtime/Engine/Private/Particles/ParticlePerfStats.cpp
--- a/Engine/Source/Runtime/Engine/Private/Particles/ParticlePerfStats.cpp
+++ b/Engine/Source/Runtime/Engine/Private/Particles/ParticlePerfStats.cpp
@@ -76,6 +76,18 @@ static FAutoConsoleVariableRef CVarParticlePerfStatsEnabled(
 	ECVF_Default
 );
 
+int32 GParticlePerfStatsMemoryMode = 1;
+static FAutoConsoleVariableRef CVarFXDetailedCSVMemoryStats(
+	TEXT("fx.ParticlePerfStats.MemoryMode"),
+	GParticlePerfStatsMemoryMode,
+	TEXT("Gathers approximate memory information depending on the mode.\n")
+	TEXT("0 = Disabled.\n")
+	TEXT("1 = Component Approximate Size (default - low perf impact).\n")
+	TEXT("2 = Asset Size + Component Approximate Size (low perf impact).\n")
+	TEXT("3 = Asset Size + Component Resource Size (high performance impact)."),
+	ECVF_Default
+);
+
 static FAutoConsoleCommandWithWorldAndArgs GParticlePerfStatsRunTest(
 	TEXT("fx.ParticlePerfStats.RunTest"),
 	TEXT("Runs for a number of frames then logs out the results.\nArg0 = NumFrames.\nArg1 = Gather World Stats (default 0).\nArg2 = Gather System Stats (default 1).\nArg3 = Gather Component Stats (default 0)."),
@@ -133,7 +145,7 @@ void FParticlePerfStatsManager::AddListener(FParticlePerfStatsListenerPtr Listen
 		{
 			FParticlePerfStats::AddWorldStatReader();
 		}
-		if (Listener->NeedsSystemStats())
+		if (Listener->NeedsSystemStats() || Listener->NeedsMemoryStats())
 		{
 			FParticlePerfStats::AddSystemStatReader();
 		}
@@ -168,7 +180,7 @@ void FParticlePerfStatsManager::RemoveListener(FParticlePerfStatsListenerPtr Lis
 	{
 		FParticlePerfStats::RemoveWorldStatReader();
 	}
-	if (Listener->NeedsSystemStats())
+	if (Listener->NeedsSystemStats() || Listener->NeedsMemoryStats())
 	{
 		FParticlePerfStats::RemoveSystemStatReader();
 	}
@@ -236,6 +248,62 @@ void FParticlePerfStatsManager::Tick()
 
 		check(Listeners.Num() > 0);
 
+#if WITH_PER_SYSTEM_PARTICLE_PERF_STATS
+		// If anyone needs memory stats we tick those now
+		if (GParticlePerfStatsMemoryMode != 0 && Listeners.ContainsByPredicate([](const FParticlePerfStatsListenerPtr& Listener) { return Listener->NeedsMemoryStats();}))
+		{
+			checkSlow(FParticlePerfStats::GetGatherSystemStats() && FParticlePerfStats::GetStatsEnabled());
+			const bool bIncludeAssetSize = GParticlePerfStatsMemoryMode > 1;
+			const bool bComponentUseResourceSize = GParticlePerfStatsMemoryMode == 3;
+
+			// Gather per loaded asset cost
+			if (bIncludeAssetSize)
+			{
+				for (TObjectIterator<UFXSystemAsset> It; It; ++It)
+				{
+					UFXSystemAsset* FXAsset = *It;
+					if (!IsValid(FXAsset) || FXAsset->IsUnreachable() || FXAsset->HasAnyFlags(EObjectFlags::RF_ClassDefaultObject))
+					{
+						continue;
+					}
+
+					FParticlePerfStats* Stats = FParticlePerfStatsManager::GetSystemPerfStats(FXAsset);
+					if (!Stats->MemoryKB_Asset.IsSet())
+					{
+						FResourceSizeEx ResSize = FResourceSizeEx(EResourceSizeMode::EstimatedTotal);
+						FXAsset->GetResourceSizeEx(ResSize);
+						Stats->MemoryKB_Asset = ResSize.GetTotalMemoryBytes();
+					}
+				}
+			}
+
+			// Gather per component cost
+			for (TObjectIterator<UFXSystemComponent> It; It; ++It)
+			{
+				UFXSystemComponent* FXComponent = *It;
+				UFXSystemAsset* FXAsset = IsValid(FXComponent) ? FXComponent->GetFXSystemAsset() : nullptr;
+				if (!IsValid(FXAsset) || FXComponent->IsUnreachable() || FXComponent->HasAnyFlags(EObjectFlags::RF_ClassDefaultObject))
+				{
+					continue;
+				}
+
+				FParticlePerfStats* Stats = FParticlePerfStatsManager::GetSystemPerfStats(FXAsset);
+				if (bComponentUseResourceSize)
+				{
+					Stats->GetGameThreadStats().MemoryBytes += FArchiveCountMem(FXComponent).GetMax();
+
+					FResourceSizeEx CompResSize = FResourceSizeEx(EResourceSizeMode::EstimatedTotal);
+					FXComponent->GetResourceSizeEx(CompResSize);
+					Stats->GetGameThreadStats().MemoryBytes += CompResSize.GetTotalMemoryBytes();
+				}
+				else
+				{
+					Stats->GetGameThreadStats().MemoryBytes += FXComponent->GetApproxMemoryUsage();
+				}
+			}
+		}
+#endif
+
 		//Tick our listeners so they can consume the finished frame data.
 		TArray<FParticlePerfStatsListenerPtr, TInlineAllocator<8>> ToRemove;
 		for (FParticlePerfStatsListenerPtr& Listener : Listeners)
@@ -321,6 +389,7 @@ void FParticlePerfStatsManager::Tick()
 #if WITH_PARTICLE_PERF_CSV_STATS
 					it.Value()->ResetStatNames();
 #endif
+					it.Value()->MemoryKB_Asset.Reset();	// Note not clear above as this path can be taken without CSV stats
 					FreeSystemStatsPool.Emplace(it.Value().Release());
 					for (FParticlePerfStatsListenerPtr& Listener : Listeners)
 					{
@@ -1160,18 +1229,6 @@ static FAutoConsoleVariable CVarWriteDetailedCSVStats(
 	ECVF_Default | ECVF_RenderThreadSafe
 );
 
-int32 GFXDetailedCSVMemoryMode = 1;
-static FAutoConsoleVariableRef CVarFXDetailedCSVMemoryStats(
-	TEXT("fx.DetailedCSVStats.MemoryMode"),
-	GFXDetailedCSVMemoryMode,
-	TEXT("Gathers approximate memory information depending on the mode.\n")
-	TEXT("0 = Disabled (default).\n")
-	TEXT("1 = Component Approximate Size (low perf impact).\n")
-	TEXT("2 = Asset Size + Component Approximate Size (low perf impact).\n")
-	TEXT("3 = Asset Size + Component Resource Size (high performance impact)."),
-	ECVF_Default
-);
-
 void OnDetailedCSVStatsEnabledChanged(IConsoleVariable* Variable)
 {
 	const bool bDetailedCSVStats = CVarWriteDetailedCSVStats->GetBool();
@@ -1227,75 +1284,6 @@ bool FParticlePerfStatsListener_CSVProfiler::Tick()
 	if (FCsvProfiler* CSVProfiler = FCsvProfiler::Get())
 	{
 #if WITH_PER_SYSTEM_PARTICLE_PERF_STATS
-		// Very slow, this gives coverage over everything but is bad for performance
-		if (GFXDetailedCSVMemoryMode != 0)
-		{
-			TMap<UFXSystemAsset*, uint64> MemoryUsage;
-
-			const bool bIncludeAssetSize = GFXDetailedCSVMemoryMode > 1;
-			const bool bComponentUseResourceSize = GFXDetailedCSVMemoryMode == 3;
-
-			// Gather per loaded asset cost
-			if (bIncludeAssetSize)
-			{
-				for (TObjectIterator<UFXSystemAsset> It; It; ++It)
-				{
-					UFXSystemAsset* FXAsset = *It;
-					if (!IsValid(FXAsset) || FXAsset->IsUnreachable() || FXAsset->HasAnyFlags(EObjectFlags::RF_ClassDefaultObject))
-					{
-						continue;
-					}
-
-					FParticlePerfStats* Stats = FParticlePerfStatsManager::GetSystemPerfStats(FXAsset);
-					if (!Stats->CSVMemoryKB_Asset.IsSet())
-					{
-						FResourceSizeEx ResSize = FResourceSizeEx(EResourceSizeMode::EstimatedTotal);
-						FXAsset->GetResourceSizeEx(ResSize);
-						Stats->CSVMemoryKB_Asset = ResSize.GetTotalMemoryBytes();
-					}
-					MemoryUsage.FindOrAdd(FXAsset) = Stats->CSVMemoryKB_Asset.GetValue();
-				}
-			}
-
-			// Gather per component cost
-			for (TObjectIterator<UFXSystemComponent> It; It; ++It)
-			{
-				UFXSystemComponent* FXComponent = *It;
-				UFXSystemAsset* FXAsset = FXComponent->GetFXSystemAsset();
-				if (!IsValidChecked(FXComponent) || !IsValid(FXAsset) || FXComponent->IsUnreachable() || FXComponent->HasAnyFlags(EObjectFlags::RF_ClassDefaultObject))
-				{
-					continue;
-				}
-
-				FParticlePerfStats* Stats = FParticlePerfStatsManager::GetSystemPerfStats(FXAsset);
-				if (Stats->CSVStat_MemoryKB.IsNone())
-				{
-					continue;
-				}
-
-				uint64& MemoryBytes = MemoryUsage.FindOrAdd(FXAsset);
-				if (bComponentUseResourceSize)
-				{
-					MemoryBytes += FArchiveCountMem(FXComponent).GetMax();
-
-					FResourceSizeEx CompResSize = FResourceSizeEx(EResourceSizeMode::EstimatedTotal);
-					FXComponent->GetResourceSizeEx(CompResSize);
-					MemoryBytes += CompResSize.GetTotalMemoryBytes();
-				}
-				else
-				{
-					MemoryBytes += FXComponent->GetApproxMemoryUsage();
-				}
-			}
-
-			for (auto OutputIt=MemoryUsage.CreateConstIterator(); OutputIt; ++OutputIt)
-			{
-				const int32 MemoryKB = (int32)FMath::DivideAndRoundUp(OutputIt.Value(), 1024ull);
-				FParticlePerfStats* Stats = FParticlePerfStatsManager::GetSystemPerfStats(OutputIt.Key());
-				CSVProfiler->RecordCustomStat(Stats->CSVStat_MemoryKB, CSV_CATEGORY_INDEX(Particles), MemoryKB, ECsvCustomStatOp::Set);
-			}
-		}
-
 		FParticlePerfStatsManager::ForAllSystemStats(
 			[&](TWeakObjectPtr<const UFXSystemAsset>& WeakSystem, TUniquePtr<FParticlePerfStats>& Stats)
 			{
@@ -1319,6 +1307,13 @@ bool FParticlePerfStatsListener_CSVProfiler::Tick()
 
 						CSVProfiler->RecordCustomStat(Stats->CSVStat_Activation, CSV_CATEGORY_INDEX(Particles), Activation, ECsvCustomStatOp::Set);
 						CSVProfiler->RecordCustomStat(Stats->CSVStat_Waits, CSV_CATEGORY_INDEX(Particles), Wait, ECsvCustomStatOp::Set);
+
+						if (!Stats->CSVStat_MemoryKB.IsNone())
+						{
+							const uint64 MemoryBytes = Stats->GetGameThreadStats().MemoryBytes + Stats->MemoryKB_Asset.Get(0);
+							const int32 MemoryKB = (int32)FMath::DivideAndRoundUp(MemoryBytes, 1024ull);
+							CSVProfiler->RecordCustomStat(Stats->CSVStat_MemoryKB, CSV_CATEGORY_INDEX(Particles), MemoryKB, ECsvCustomStatOp::Set);
+						}
 					}
 				}
 			}
@@ -1406,8 +1401,8 @@ int32 FParticlePerfStatsListener_DebugRender::RenderStats(class UWorld* World, c
 	// Display Header
 	Canvas->DrawTile(X - 2, Y - 1, (ColumnWidth * 5) + 4, FontHeight, 0.0f, 0.0f, 1.0f, 1.0f, HeaderBackground);
 	Canvas->DrawShadowedString(X + ColumnWidth * 0, Y, TEXT("System Name"), Font, FLinearColor::Yellow);
-	Canvas->DrawShadowedString(X + ColumnWidth * 1, Y, TEXT("Average PerFrame GT | GT CNC | RT"), Font, FLinearColor::Yellow);
-	Canvas->DrawShadowedString(X + ColumnWidth * 2, Y, TEXT("Average PerInstance GT | GT CNC | RT"), Font, FLinearColor::Yellow);
+	Canvas->DrawShadowedString(X + ColumnWidth * 1, Y, TEXT("Average PerFrame GT | GT CNC | RT | GPU"), Font, FLinearColor::Yellow);
+	Canvas->DrawShadowedString(X + ColumnWidth * 2, Y, TEXT("Average PerInstance GT | GT CNC | RT | GPU"), Font, FLinearColor::Yellow);
 	Canvas->DrawShadowedString(X + ColumnWidth * 3, Y, TEXT("Peak PerFrame GT | RT"), Font, FLinearColor::Yellow);
 	Y += FontHeight;
 
@@ -1423,6 +1418,7 @@ int32 FParticlePerfStatsListener_DebugRender::RenderStats(class UWorld* World, c
 
 		const FAccumulatedParticlePerfStats_GT& GTStats = PerfStats->GetGameThreadStats();
 		const FAccumulatedParticlePerfStats_RT& RTStats = PerfStats->GetRenderThreadStats();
+		const FAccumulatedParticlePerfStats_GPU& GPUStats = PerfStats->GetGPUStats();
 		const UFXSystemAsset* System = it.Key().Get();
 		if (GTStats.NumFrames == 0 || RTStats.NumFrames == 0 || System == nullptr)
 		{
@@ -1444,7 +1440,7 @@ int32 FParticlePerfStatsListener_DebugRender::RenderStats(class UWorld* World, c
 			uint32(FPlatformTime::ToMilliseconds64(GTStats.AccumulatedStats.TickGameThreadCycles + GTStats.AccumulatedStats.FinalizeCycles) * 1000.0 / double(GTStats.NumFrames)),
 			uint32(FPlatformTime::ToMilliseconds64(GTStats.AccumulatedStats.TickConcurrentCycles + GTStats.AccumulatedStats.EndOfFrameCycles) * 1000.0 / double(GTStats.NumFrames)),
 			uint32(FPlatformTime::ToMilliseconds64(RTStats.AccumulatedStats.RenderUpdateCycles + RTStats.AccumulatedStats.GetDynamicMeshElementsCycles) * 1000.0 / double(RTStats.NumFrames)),
-			uint32(FPlatformTime::ToMilliseconds64(GTStats.AccumulatedStats.TickGameThreadCycles) * 1000.0 / double(GTStats.NumFrames))
+			uint32((double(GPUStats.AccumulatedStats.TotalMicroseconds) / 1000.0) / double(GTStats.NumFrames))
 		);
 		Canvas->DrawShadowedString(X + ColumnWidth * 1, Y, *tempString, Font, FLinearColor::Yellow);
 
@@ -1455,7 +1451,7 @@ int32 FParticlePerfStatsListener_DebugRender::RenderStats(class UWorld* World, c
 			uint32(FPlatformTime::ToMilliseconds64(GTStats.AccumulatedStats.TickGameThreadCycles + GTStats.AccumulatedStats.FinalizeCycles) * 1000.0 / double(GTStats.AccumulatedStats.NumInstances)),
 			uint32(FPlatformTime::ToMilliseconds64(GTStats.AccumulatedStats.TickConcurrentCycles + GTStats.AccumulatedStats.EndOfFrameCycles) * 1000.0 / double(GTStats.AccumulatedStats.NumInstances)),
 			uint32(FPlatformTime::ToMilliseconds64(RTStats.AccumulatedStats.RenderUpdateCycles + RTStats.AccumulatedStats.GetDynamicMeshElementsCycles) * 1000.0 / double(RTStats.AccumulatedStats.NumInstances)),
-			uint32(FPlatformTime::ToMilliseconds64(GTStats.AccumulatedStats.TickGameThreadCycles) * 1000.0 / double(GTStats.AccumulatedStats.NumInstances))
+			uint32((double(GPUStats.AccumulatedStats.TotalMicroseconds) / 1000.0) / double(GTStats.AccumulatedStats.NumInstances))
 		);
 		Canvas->DrawShadowedString(X + ColumnWidth * 2, Y, *tempString, Font, FLinearColor::Yellow);
 
