- Add Data Channel events to Niagara Insights
- Add option to include potentially overlapping events, this is useful for connected event tracks

#Synced-CL 51808504

diff --git a/Engine/Plugins/FX/Niagara/Source/Niagara/Internal/NiagaraTrace.h b/Engine/Plugins/FX/Niagara/Source/Niagara/Internal/NiagaraTrace.h
--- a/Engine/Plugins/FX/Niagara/Source/Niagara/Internal/NiagaraTrace.h
+++ b/Engine/Plugins/FX/Niagara/Source/Niagara/Internal/NiagaraTrace.h
@@ -11,6 +11,7 @@
 #endif
 
 class UNiagaraComponent;
+struct FNiagaraDataChannelPublishRequest;
 
 UE_TRACE_CHANNEL_EXTERN(NiagaraChannel, NIAGARA_API)
 
@@ -23,6 +24,9 @@ namespace FNiagaraInsights
 	void NotifyComponentActivate(const UNiagaraComponent* Component, bool bReset, bool bIsScalabilityCull, bool bAwaitingActivationDueToNotReady);
 	void NotifyComponentDeactivate(const UNiagaraComponent* Component, bool bImmediate, bool bIsScalabilityCull);
 	void NotifyComponentComplete(const UNiagaraComponent* Component, bool bExternalCompletion);
+
+	void NotifyDataChannelPublish(const FNiagaraDataChannelPublishRequest& Request, bool bGpuRequest);
+	void NotifyDataChannelWrite(const UObject* DataChannel, const TCHAR* SourceName, int32 Count, bool bVisibleToGame, bool bVisibleToCPU, bool bVisibleToGPU);
 }
 #else
 namespace FNiagaraInsights
@@ -33,5 +37,8 @@ namespace FNiagaraInsights
 	inline void NotifyComponentActivate(const UNiagaraComponent* Component, bool bReset, bool bIsScalabilityCull, bool bAwaitingActivationDueToNotReady) {}
 	inline void NotifyComponentDeactivate(const UNiagaraComponent* Component, bool bImmediate, bool bIsScalabilityCull) {}
 	inline void NotifyComponentComplete(const UNiagaraComponent* Component, bool bExternalCompletion) {}
+
+	inline void NotifyDataChannelPublish(const FNiagaraDataChannelPublishRequest& Request, bool bGpuRequest) {}
+	inline void NotifyDataChannelWrite(const UObject* DataChannel, const TCHAR* SourceName, int32 Count, bool bVisibleToGame, bool bVisibleToCPU, bool bVisibleToGPU) {}
 }
 #endif //WITH_NIAGARA_INSIGHTS
diff --git a/Engine/Plugins/FX/Niagara/Source/Niagara/Private/NiagaraDataChannelAccessor.cpp b/Engine/Plugins/FX/Niagara/Source/Niagara/Private/NiagaraDataChannelAccessor.cpp
--- a/Engine/Plugins/FX/Niagara/Source/Niagara/Private/NiagaraDataChannelAccessor.cpp
+++ b/Engine/Plugins/FX/Niagara/Source/Niagara/Private/NiagaraDataChannelAccessor.cpp
@@ -11,6 +11,7 @@
 #include "NiagaraDataChannelManager.h"
 #include "NiagaraDataChannelData.h"
 
+#include "NiagaraTrace.h"
 #include "NiagaraWorldManager.h"
 #include "NiagaraDataChannelAccessorImpl.h"
 
@@ -201,6 +202,8 @@ bool UNiagaraDataChannelWriter::BeginWrite(FNDCAccessContextInst& AccessContext,
 {
 	check(IsInGameThread());
 
+	FNiagaraInsights::NotifyDataChannelWrite(Owner ? Owner->GetDataChannel() : nullptr, *DebugSource, Count, bVisibleToGame, bVisibleToCPU, bVisibleToGPU);
+
 	if (Count == 0)
 	{
 		UE_LOG(LogNiagara, Warning, TEXT("Call to UNiagaraDataChannelWriter::InitWrite with Count == 0. Ignored."));
@@ -402,11 +405,15 @@ void FNDCWriterBase::EndWrite()
 
 bool FNDCWriterBase::BeginWrite(const UObject* WorldContextObject, const UNiagaraDataChannel* DataChannel, FNDCAccessContextInst& AccessContext, int32 InCount, bool bVisibleToGame, bool bVisibleToCPU, bool bVisibleToGPU)
 {
+	FNiagaraInsights::NotifyDataChannelWrite(DataChannel, *DebugSource, InCount, bVisibleToGame, bVisibleToCPU, bVisibleToGPU);
+
 	return BeginWrite_Internal(WorldContextObject, DataChannel, AccessContext, InCount, bVisibleToGame, bVisibleToCPU, bVisibleToGPU);
 }
 
 bool FNDCWriterBase::BeginWrite(const UObject* WorldContextObject, const UNiagaraDataChannel* DataChannel, FNiagaraDataChannelSearchParameters SearchParams, int32 InCount, bool bVisibleToGame, bool bVisibleToCPU, bool bVisibleToGPU)
 {
+	FNiagaraInsights::NotifyDataChannelWrite(DataChannel, *DebugSource, InCount, bVisibleToGame, bVisibleToCPU, bVisibleToGPU);
+
 	return BeginWrite_Internal(WorldContextObject, DataChannel, SearchParams, InCount, bVisibleToGame, bVisibleToCPU, bVisibleToGPU);
 }
 
diff --git a/Engine/Plugins/FX/Niagara/Source/Niagara/Private/NiagaraDataChannelData.cpp b/Engine/Plugins/FX/Niagara/Source/Niagara/Private/NiagaraDataChannelData.cpp
--- a/Engine/Plugins/FX/Niagara/Source/Niagara/Private/NiagaraDataChannelData.cpp
+++ b/Engine/Plugins/FX/Niagara/Source/Niagara/Private/NiagaraDataChannelData.cpp
@@ -11,6 +11,8 @@
 #include "NiagaraGpuComputeDispatchInterface.h"
 
 #include "NiagaraDataSetReadback.h"
+#include "NiagaraTrace.h"
+
 #include "RenderGraphUtils.h"
 #include "Engine/World.h"
 
@@ -380,13 +382,16 @@ FNiagaraDataChannelGameData* FNiagaraDataChannelData::GetGameData()
 
 void FNiagaraDataChannelData::Publish(const FNiagaraDataChannelPublishRequest& Request)
 {
-	FScopeLock Lock(&PublishCritSec);
+	FNiagaraInsights::NotifyDataChannelPublish(Request, false);
 
+	FScopeLock Lock(&PublishCritSec);
 	PublishRequests.Add(Request);
 }
 
 void FNiagaraDataChannelData::PublishFromGPU(const FNiagaraDataChannelPublishRequest& Request)
 {
+	FNiagaraInsights::NotifyDataChannelPublish(Request, true);
+
 	check(IsInGameThread());
 	PublishRequestsFromGPU.Add(Request);
 }
diff --git a/Engine/Plugins/FX/Niagara/Source/Niagara/Private/NiagaraTrace.cpp b/Engine/Plugins/FX/Niagara/Source/Niagara/Private/NiagaraTrace.cpp
--- a/Engine/Plugins/FX/Niagara/Source/Niagara/Private/NiagaraTrace.cpp
+++ b/Engine/Plugins/FX/Niagara/Source/Niagara/Private/NiagaraTrace.cpp
@@ -3,6 +3,7 @@
 #include "NiagaraTrace.h"
 
 #include "NiagaraComponent.h"
+#include "NiagaraDataChannelPublishRequest.h"
 
 #include "Particles/ParticlePerfStats.h"
 
@@ -83,6 +84,27 @@ namespace NiagaraInsightsPrivate
 		UE_TRACE_EVENT_FIELD(bool,					bExternalCompletion)
 	UE_TRACE_EVENT_END()
 
+	UE_TRACE_EVENT_BEGIN(NiagaraTrace, DataChannelPublish, NoSync)
+		UE_TRACE_EVENT_FIELD(uint64,				EventTime)
+		UE_TRACE_EVENT_FIELD(UE::Trace::WideString, SourceName)
+		UE_TRACE_EVENT_FIELD(bool,					bGpuRequest)
+		UE_TRACE_EVENT_FIELD(bool,					bVisibleToGame)
+		UE_TRACE_EVENT_FIELD(bool,					bVisibleToCPUSims)
+		UE_TRACE_EVENT_FIELD(bool,					bVisibleToGPUSims)
+		UE_TRACE_EVENT_FIELD(uint32,				NumInstances)
+		UE_TRACE_EVENT_FIELD(uint32,				NumInstanceAllocated)
+	UE_TRACE_EVENT_END()
+
+	UE_TRACE_EVENT_BEGIN(NiagaraTrace, DataChannelWrite, NoSync)
+		UE_TRACE_EVENT_FIELD(uint64,				EventTime)
+		UE_TRACE_EVENT_FIELD(UE::Trace::WideString, DataChannelName)
+		UE_TRACE_EVENT_FIELD(UE::Trace::WideString, SourceName)
+		UE_TRACE_EVENT_FIELD(int32,					NumInstances)
+		UE_TRACE_EVENT_FIELD(bool,					bVisibleToGame)
+		UE_TRACE_EVENT_FIELD(bool,					bVisibleToCPU)
+		UE_TRACE_EVENT_FIELD(bool,					bVisibleToGPU)
+	UE_TRACE_EVENT_END()
+
 #if WITH_PARTICLE_PERF_STATS
 	class FNiagaraInsightsStatsListener : public FParticlePerfStatsListener
 	{
@@ -240,4 +262,47 @@ void FNiagaraInsights::NotifyComponentComplete(const UNiagaraComponent* Componen
 	}
 }
 
+void FNiagaraInsights::NotifyDataChannelPublish(const FNiagaraDataChannelPublishRequest& Request, bool bGpuRequest)
+{
+	using namespace NiagaraInsightsPrivate;
+	if (IsChannelActive())
+	{
+		UE_TRACE_LOG(NiagaraTrace, DataChannelPublish, NiagaraChannel)
+			<< DataChannelPublish.EventTime(FPlatformTime::Cycles64())
+		#if WITH_NIAGARA_DEBUGGER
+			<< DataChannelPublish.SourceName(*Request.DebugSource)
+		#endif
+			<< DataChannelPublish.bGpuRequest(bGpuRequest)
+			<< DataChannelPublish.bVisibleToGame(Request.bVisibleToGame)
+			<< DataChannelPublish.bVisibleToCPUSims(Request.bVisibleToCPUSims)
+			<< DataChannelPublish.bVisibleToGPUSims(Request.bVisibleToGPUSims)
+			<< DataChannelPublish.NumInstances(Request.Data.IsValid() ? Request.Data->GetNumInstances() : 0)
+			<< DataChannelPublish.NumInstanceAllocated(Request.Data.IsValid() ? Request.Data->GetNumInstancesAllocated() : 0)
+			;
+	}
+}
+
+void FNiagaraInsights::NotifyDataChannelWrite(const UObject* DataChannel, const TCHAR* SourceName, int32 Count, bool bVisibleToGame, bool bVisibleToCPU, bool bVisibleToGPU)
+{
+	using namespace NiagaraInsightsPrivate;
+	if (IsChannelActive())
+	{
+		FNameBuilder DataChannelName;
+		if (DataChannel)
+		{
+			DataChannel->GetFName().ToString(DataChannelName);
+		}
+
+		UE_TRACE_LOG(NiagaraTrace, DataChannelWrite, NiagaraChannel)
+			<< DataChannelWrite.EventTime(FPlatformTime::Cycles64())
+			<< DataChannelWrite.DataChannelName(DataChannelName.ToString())
+			<< DataChannelWrite.SourceName(SourceName ? SourceName : TEXT(""))
+			<< DataChannelWrite.NumInstances(Count)
+			<< DataChannelWrite.bVisibleToGame(bVisibleToGame)
+			<< DataChannelWrite.bVisibleToCPU(bVisibleToCPU)
+			<< DataChannelWrite.bVisibleToGPU(bVisibleToGPU)
+			;
+	}
+}
+
 #endif //WITH_NIAGARA_INSIGHTS
diff --git a/Engine/Plugins/FX/NiagaraInsights/Source/NiagaraInsights/Private/NiagaraAnalyzer.cpp b/Engine/Plugins/FX/NiagaraInsights/Source/NiagaraInsights/Private/NiagaraAnalyzer.cpp
--- a/Engine/Plugins/FX/NiagaraInsights/Source/NiagaraInsights/Private/NiagaraAnalyzer.cpp
+++ b/Engine/Plugins/FX/NiagaraInsights/Source/NiagaraInsights/Private/NiagaraAnalyzer.cpp
@@ -33,6 +33,9 @@ void FNiagaraAnalyzer::OnAnalysisBegin(const FOnAnalysisContext& Context)
 	Builder.RouteEvent(RouteId_ComponentActivate,	"NiagaraTrace", "ComponentActivate");
 	Builder.RouteEvent(RouteId_ComponentDeactivate,	"NiagaraTrace", "ComponentDeactivate");
 	Builder.RouteEvent(RouteId_ComponentComplete,	"NiagaraTrace", "ComponentComplete");
+
+	Builder.RouteEvent(RouteId_DataChannelPublish, "NiagaraTrace", "DataChannelPublish");
+	Builder.RouteEvent(RouteId_DataChannelWrite, "NiagaraTrace", "DataChannelWrite");
 }
 
 bool FNiagaraAnalyzer::OnEvent(uint16 RouteId, EStyle Style, const FOnEventContext& Context)
@@ -108,6 +111,35 @@ bool FNiagaraAnalyzer::OnEvent(uint16 RouteId, EStyle Style, const FOnEventConte
 			);
 			break;
 		}
+
+		case RouteId_DataChannelPublish:
+		{
+			Provider.AddDataChannelPublish(
+				EventTime,
+				GetEventString(EventData, "SourceName"),
+				EventData.GetValue<bool>("bGpuRequest"),
+				EventData.GetValue<bool>("bVisibleToGame"),
+				EventData.GetValue<bool>("bVisibleToCPUSims"),
+				EventData.GetValue<bool>("bVisibleToGPUSims"),
+				EventData.GetValue<uint32>("NumInstances"),
+				EventData.GetValue<uint32>("NumInstanceAllocated")
+			);
+			break;
+		}
+
+		case RouteId_DataChannelWrite:
+		{
+			Provider.AddDataChannelWrite(
+				EventTime,
+				GetEventString(EventData, "DataChannelName"),
+				GetEventString(EventData, "SourceName"),
+				EventData.GetValue<int32>("NumInstances"),
+				EventData.GetValue<bool>("bVisibleToGame"),
+				EventData.GetValue<bool>("bVisibleToCPU"),
+				EventData.GetValue<bool>("bVisibleToGPU")
+			);
+			break;
+		}
 	}
 	return true;
 }
diff --git a/Engine/Plugins/FX/NiagaraInsights/Source/NiagaraInsights/Private/NiagaraAnalyzer.h b/Engine/Plugins/FX/NiagaraInsights/Source/NiagaraInsights/Private/NiagaraAnalyzer.h
--- a/Engine/Plugins/FX/NiagaraInsights/Source/NiagaraInsights/Private/NiagaraAnalyzer.h
+++ b/Engine/Plugins/FX/NiagaraInsights/Source/NiagaraInsights/Private/NiagaraAnalyzer.h
@@ -29,6 +29,9 @@ private:
 		RouteId_ComponentActivate,
 		RouteId_ComponentDeactivate,
 		RouteId_ComponentComplete,
+
+		RouteId_DataChannelPublish,
+		RouteId_DataChannelWrite,
 	};
 
 	TraceServices::IAnalysisSession&	Session;
diff --git a/Engine/Plugins/FX/NiagaraInsights/Source/NiagaraInsights/Private/NiagaraDataChannelTrack.cpp b/Engine/Plugins/FX/NiagaraInsights/Source/NiagaraInsights/Private/NiagaraDataChannelTrack.cpp
new file mode 100644
--- /dev/null
+++ b/Engine/Plugins/FX/NiagaraInsights/Source/NiagaraInsights/Private/NiagaraDataChannelTrack.cpp
@@ -0,0 +1,138 @@
+// Copyright Epic Games, Inc. All Rights Reserved.
+
+#include "NiagaraDataChannelTrack.h"
+
+#include "NiagaraInsightsCommon.h"
+#include "NiagaraProvider.h"
+#include "NiagaraTimingViewSession.h"
+
+#include "Insights/ViewModels/GraphTrackBuilder.h"
+#include "Insights/ViewModels/GraphTrackEvent.h"
+#include "Insights/ViewModels/TimingTrackViewport.h"
+#include "Insights/ViewModels/TooltipDrawState.h"
+
+#define LOCTEXT_NAMESPACE "NiagaraDataChannelTrack"
+
+namespace UE::NiagaraInsights
+{
+
+INSIGHTS_IMPLEMENT_RTTI(FNiagaraDataChannelTrack)
+
+FNiagaraDataChannelTrack::FNiagaraDataChannelTrack(const FNiagaraTimingViewSession& InSharedData)
+	: FGraphTrack(LOCTEXT("NiagaraDataChannel", "Niagara Data Channel").ToString())
+	, SharedData(InSharedData)
+{
+	constexpr double NumEventTypes = 2.0;
+
+	EnabledOptions = //EGraphOptions::ShowDebugInfo |
+		EGraphOptions::ShowPoints |
+		EGraphOptions::ShowLines |
+		EGraphOptions::ShowPolygon |
+		EGraphOptions::ShowVerticalAxisGrid |
+		EGraphOptions::ShowHeader;
+
+	SetValidLocations(ETimingTrackLocation::Scrollable | ETimingTrackLocation::TopDocked | ETimingTrackLocation::BottomDocked);
+	SetHeight(NumEventTypes * 10.0);
+
+	TSharedRef<FGraphSeries> GraphSeries = MakeShared<FGraphSeries>();
+	GraphSeries->SetName(TEXT("Events"));
+	GraphSeries->SetColor(Common::Color_MemoryCostTotal);
+	GraphSeries->SetBaselineY(GetHeight());
+	GraphSeries->SetScaleY(GetHeight() / (NumEventTypes + 1.0));
+	GraphSeries->EnableAutoZoom();
+	GetSeries().Add(GraphSeries);
+}
+
+void FNiagaraDataChannelTrack::Update(const ITimingTrackUpdateContext& Context)
+{
+	Super::Update(Context);
+
+	const FNiagaraProvider* Provider = SharedData.GetAnalysisSession().ReadProvider<FNiagaraProvider>(FNiagaraProvider::GetProviderName());
+	if (Provider == nullptr)
+	{
+		return;
+	}
+
+	const FTimingTrackViewport& Viewport = Context.GetViewport();
+
+	const TSharedPtr<FGraphSeries> GraphSeries = GetSeries()[0];
+	if (GraphSeries->IsVisible())
+	{
+		FGraphTrackBuilder Builder(*this, *GraphSeries, Viewport);
+		Provider->EnumerateDataChannelEvent(
+			Viewport.GetStartTime(),
+			Viewport.GetEndTime(),
+			[&Builder, &GraphSeries](const FDataChannelEvent& DataChannelEvent)
+			{
+				const double Offset = double(DataChannelEvent.Payload.GetIndex()) + 0.5;
+				Builder.AddEvent(DataChannelEvent.Time, 0.0, Offset, false);
+			}
+		);
+	}
+}
+
+void FNiagaraDataChannelTrack::InitTooltip(FTooltipDrawState& InOutTooltip, const ITimingEvent& InTooltipEvent) const
+{
+	if (!InTooltipEvent.CheckTrack(this) || !InTooltipEvent.Is<FGraphTrackEvent>())
+	{
+		return;
+	}
+
+	const FGraphTrackEvent& GraphEvent = InTooltipEvent.As<FGraphTrackEvent>();
+	const TSharedRef<const FGraphSeries> HoveredSeries = GraphEvent.GetSeries();
+
+	const FNiagaraProvider* Provider = SharedData.GetAnalysisSession().ReadProvider<FNiagaraProvider>(FNiagaraProvider::GetProviderName());
+	if (Provider == nullptr)
+	{
+		return;
+	}
+
+	InOutTooltip.ResetContent();
+	InOutTooltip.AddTitle(HoveredSeries->GetName().ToString(), HoveredSeries->GetColor());
+
+	const double EventTime = GraphEvent.GetStartTime();
+	constexpr double TimeTolerance = 1e-9;
+
+	Provider->EnumerateDataChannelEvent(
+		EventTime - TimeTolerance,
+		EventTime + TimeTolerance,
+		[&InOutTooltip](const FDataChannelEvent& DataChannelEvent)
+		{
+			switch (DataChannelEvent.Payload.GetIndex())
+			{
+				case 0:
+				{
+					const FDataChannelEvent::FPublish& Payload = DataChannelEvent.Payload.Get<FDataChannelEvent::FPublish>();
+					InOutTooltip.AddNameValueTextLine(TEXT("Publish"),				TEXT(""));
+					InOutTooltip.AddNameValueTextLine(TEXT("SourceName"),			Payload.SourceName.IsEmpty() ? TEXT("<No Data>") : Payload.SourceName);
+					InOutTooltip.AddNameValueTextLine(TEXT("bGpuRequest"),			LexToString(Payload.bGpuRequest));
+					InOutTooltip.AddNameValueTextLine(TEXT("bVisibleToGame"),		LexToString(Payload.bVisibleToGame));
+					InOutTooltip.AddNameValueTextLine(TEXT("bVisibleToCPUSims"),	LexToString(Payload.bVisibleToCPUSims));
+					InOutTooltip.AddNameValueTextLine(TEXT("bVisibleToGPUSims"),	LexToString(Payload.bVisibleToGPUSims));
+					InOutTooltip.AddNameValueTextLine(TEXT("NumInstances"),			LexToString(Payload.NumInstances));
+					InOutTooltip.AddNameValueTextLine(TEXT("NumInstanceAllocated"),	LexToString(Payload.NumInstanceAllocated));
+					break;
+				}
+
+				case 1:
+				{
+					const FDataChannelEvent::FWrite& Payload = DataChannelEvent.Payload.Get<FDataChannelEvent::FWrite>();
+					InOutTooltip.AddNameValueTextLine(TEXT("Write"),				TEXT(""));
+					InOutTooltip.AddNameValueTextLine(TEXT("DataChannelName"),		Payload.DataChannelName);
+					InOutTooltip.AddNameValueTextLine(TEXT("SourceName"),			Payload.SourceName.IsEmpty() ? TEXT("<No Data>") : *Payload.SourceName);
+					InOutTooltip.AddNameValueTextLine(TEXT("NumInstances"),			LexToString(Payload.NumInstances));
+					InOutTooltip.AddNameValueTextLine(TEXT("bVisibleToGame"),		LexToString(Payload.bVisibleToGame));
+					InOutTooltip.AddNameValueTextLine(TEXT("bVisibleToCPU"),		LexToString(Payload.bVisibleToCPU));
+					InOutTooltip.AddNameValueTextLine(TEXT("bVisibleToGPU"),		LexToString(Payload.bVisibleToGPU));
+					break;
+				}
+			}
+		}
+	);
+
+	InOutTooltip.UpdateLayout();
+}
+
+} //namespace UE::NiagaraInsights
+
+#undef LOCTEXT_NAMESPACE
diff --git a/Engine/Plugins/FX/NiagaraInsights/Source/NiagaraInsights/Private/NiagaraDataChannelTrack.h b/Engine/Plugins/FX/NiagaraInsights/Source/NiagaraInsights/Private/NiagaraDataChannelTrack.h
new file mode 100644
--- /dev/null
+++ b/Engine/Plugins/FX/NiagaraInsights/Source/NiagaraInsights/Private/NiagaraDataChannelTrack.h
@@ -0,0 +1,30 @@
+// Copyright Epic Games, Inc. All Rights Reserved.
+
+#pragma once
+
+#include "Insights/ViewModels/GraphTrack.h"
+
+namespace UE::NiagaraInsights
+{
+
+class FNiagaraTimingViewSession;
+
+class FNiagaraDataChannelTrack : public FGraphTrack
+{
+	using Super = FGraphTrack;
+
+	INSIGHTS_DECLARE_RTTI(FNiagaraDataChannelTrack, FGraphTrack)
+
+public:
+	explicit FNiagaraDataChannelTrack(const FNiagaraTimingViewSession& InSharedData);
+
+	// BEGIN: FGraphTrack interface
+	virtual void Update(const ITimingTrackUpdateContext& Context) override;
+	virtual void InitTooltip(FTooltipDrawState& InOutTooltip, const ITimingEvent& InTooltipEvent) const override;
+	// END: FGraphTrack interface
+
+private:
+	const FNiagaraTimingViewSession& SharedData;
+};
+
+} //namespace UE::NiagaraInsights
diff --git a/Engine/Plugins/FX/NiagaraInsights/Source/NiagaraInsights/Private/NiagaraPerformanceGraphTrack.cpp b/Engine/Plugins/FX/NiagaraInsights/Source/NiagaraInsights/Private/NiagaraPerformanceGraphTrack.cpp
--- a/Engine/Plugins/FX/NiagaraInsights/Source/NiagaraInsights/Private/NiagaraPerformanceGraphTrack.cpp
+++ b/Engine/Plugins/FX/NiagaraInsights/Source/NiagaraInsights/Private/NiagaraPerformanceGraphTrack.cpp
@@ -112,21 +112,11 @@ void FNiagaraPerformanceGraphTrack::Update(const ITimingTrackUpdateContext& Cont
 		Provider->EnumeratePerformance_GT(
 			Viewport.GetStartTime(),
 			Viewport.GetEndTime(),
+			true,
 			[&Builder, &MaxValue](const FSystemPerformanceFrame_GT& Frame)
 			{
-				double TotalTime = 0.0;
-				for (const FSystemPerformanceFrame_GT::FSystemData& SystemData : Frame.SystemData)
-				{
-					TotalTime +=
-						SystemData.TickGameThreadSeconds +
-						SystemData.TickConcurrentSeconds +
-						SystemData.FinalizeSeconds +
-						SystemData.EndOfFrameSeconds +
-						SystemData.ActivationSeconds +
-						SystemData.WaitSeconds;
-
-				}
-				const double Value = TotalTime * 1000.0;
+				const FSystemPerformanceFrame_GT::FStats& A = Frame.AccumulatedStats;
+				const double Value = (A.TickGameThreadSeconds + A.TickConcurrentSeconds + A.FinalizeSeconds + A.EndOfFrameSeconds + A.ActivationSeconds + A.WaitSeconds) * 1000.0;
 				MaxValue = FMath::Max(MaxValue, Value);
 				Builder.AddEvent(Frame.Time, 0.0, Value);
 			}
@@ -142,14 +132,10 @@ void FNiagaraPerformanceGraphTrack::Update(const ITimingTrackUpdateContext& Cont
 		Provider->EnumeratePerformance_RT(
 			Viewport.GetStartTime(),
 			Viewport.GetEndTime(),
+			true,
 			[&Builder, &MaxValue](const FSystemPerformanceFrame_RT& Frame)
 			{
-				double TotalTime = 0.0;
-				for (const FSystemPerformanceFrame_RT::FSystemData& SystemData : Frame.SystemData)
-				{
-					TotalTime += SystemData.RenderUpdateSeconds + SystemData.GetDynamicMeshElementsSeconds;
-				}
-				const double Value = TotalTime * 1000.0;
+				const double Value = (Frame.AccumulatedStats.RenderUpdateSeconds + Frame.AccumulatedStats.GetDynamicMeshElementsSeconds) * 1000.0;
 				MaxValue = FMath::Max(MaxValue, Value);
 				Builder.AddEvent(Frame.Time, 0.0, Value);
 			}
@@ -165,14 +151,10 @@ void FNiagaraPerformanceGraphTrack::Update(const ITimingTrackUpdateContext& Cont
 		Provider->EnumeratePerformance_RT(
 			Viewport.GetStartTime(),
 			Viewport.GetEndTime(),
+			true,
 			[&Builder, &MaxValue](const FSystemPerformanceFrame_RT& Frame)
 			{
-				double TotalTime = 0.0;
-				for (const FSystemPerformanceFrame_RT::FSystemData& SystemData : Frame.SystemData)
-				{
-					TotalTime += SystemData.GpuTotalMicroseconds;
-				}
-				const double Value = TotalTime / 1000.0;
+				const double Value = Frame.AccumulatedStats.GpuTotalMicroseconds / 1000.0;
 				MaxValue = FMath::Max(MaxValue, Value);
 				Builder.AddEvent(Frame.Time, 0.0, Value);
 			}
@@ -188,15 +170,10 @@ void FNiagaraPerformanceGraphTrack::Update(const ITimingTrackUpdateContext& Cont
 		Provider->EnumeratePerformance_GT(
 			Viewport.GetStartTime(),
 			Viewport.GetEndTime(),
+			true,
 			[&Builder, &MaxValue](const FSystemPerformanceFrame_GT& Frame)
 			{
-				uint64 MemoryTotal = 0;
-				for (const FSystemPerformanceFrame_GT::FSystemData& SystemData : Frame.SystemData)
-				{
-					MemoryTotal += SystemData.MemoryBytes;
-
-				}
-				const double Value = double(MemoryTotal) / 1024.0;
+				const double Value = double(Frame.AccumulatedStats.MemoryBytes) / 1024.0;
 				MaxValue = FMath::Max(MaxValue, Value);
 				Builder.AddEvent(Frame.Time, 0.0, Value);
 			}
@@ -231,23 +208,28 @@ void FNiagaraPerformanceGraphTrack::InitTooltip(FTooltipDrawState& InOutTooltip,
 
 	if (HoveredSeries == GameTimeSeries)
 	{
-		Provider->EnumeratePerformance_GT(EventTime - TimeTolerance, EventTime + TimeTolerance,
+		Provider->EnumeratePerformance_GT(
+			EventTime - TimeTolerance,
+			EventTime + TimeTolerance,
+			false,
 			[&InOutTooltip](const FSystemPerformanceFrame_GT& Frame)
 			{
 				struct FSystemEntry { FStringView Name; double TotalMs; uint64 NumInstances; };
 
 				TArray<FSystemEntry> Entries;
-				double GrandTotalMs = 0.0;
 
-				for (const FSystemPerformanceFrame_GT::FSystemData& Sys : Frame.SystemData)
+				for (const TPair<FString, FSystemPerformanceFrame_GT::FStats>& Pair : Frame.SystemData)
 				{
+					const FSystemPerformanceFrame_GT::FStats& Sys = Pair.Value;
 					const double TotalMs = (Sys.TickGameThreadSeconds + Sys.TickConcurrentSeconds +
 						Sys.FinalizeSeconds + Sys.EndOfFrameSeconds +
 						Sys.ActivationSeconds + Sys.WaitSeconds) * 1000.0;
-					Entries.Add({ Sys.SystemName, TotalMs, Sys.NumInstances });
-					GrandTotalMs += TotalMs;
+					Entries.Add({ Pair.Key, TotalMs, Sys.NumInstances });
 				}
 
+				const FSystemPerformanceFrame_GT::FStats& A = Frame.AccumulatedStats;
+				const double GrandTotalMs = (A.TickGameThreadSeconds + A.TickConcurrentSeconds + A.FinalizeSeconds + A.EndOfFrameSeconds + A.ActivationSeconds + A.WaitSeconds) * 1000.0;
+
 				Entries.Sort([](const FSystemEntry& A, const FSystemEntry& B) { return A.TotalMs > B.TotalMs; });
 
 				InOutTooltip.AddNameValueTextLine(TEXT("Total:"), FString::Printf(TEXT("%.3f ms"), GrandTotalMs));
@@ -270,21 +252,26 @@ void FNiagaraPerformanceGraphTrack::InitTooltip(FTooltipDrawState& InOutTooltip,
 	}
 	else if (HoveredSeries == RenderTimeSeries)
 	{
-		Provider->EnumeratePerformance_RT(EventTime - TimeTolerance, EventTime + TimeTolerance,
+		Provider->EnumeratePerformance_RT(
+			EventTime - TimeTolerance,
+			EventTime + TimeTolerance,
+			false,
 			[&InOutTooltip](const FSystemPerformanceFrame_RT& Frame)
 			{
 				struct FSystemEntry { FStringView Name; double TotalMs; uint64 NumInstances; };
 
 				TArray<FSystemEntry> Entries;
-				double GrandTotalMs = 0.0;
 
-				for (const FSystemPerformanceFrame_RT::FSystemData& Sys : Frame.SystemData)
+				for (const TPair<FString, FSystemPerformanceFrame_RT::FStats>& Pair : Frame.SystemData)
 				{
+					const FSystemPerformanceFrame_RT::FStats& Sys = Pair.Value;
 					const double TotalMs = (Sys.RenderUpdateSeconds + Sys.GetDynamicMeshElementsSeconds) * 1000.0;
-					Entries.Add({ Sys.SystemName, TotalMs, Sys.NumInstances });
-					GrandTotalMs += TotalMs;
+					Entries.Add({ Pair.Key, TotalMs, Sys.NumInstances });
 				}
 
+				const FSystemPerformanceFrame_RT::FStats& A = Frame.AccumulatedStats;
+				const double GrandTotalMs = (A.RenderUpdateSeconds + A.GetDynamicMeshElementsSeconds) * 1000.0;
+
 				Entries.Sort([](const FSystemEntry& A, const FSystemEntry& B) { return A.TotalMs > B.TotalMs; });
 
 				InOutTooltip.AddNameValueTextLine(TEXT("Total:"), FString::Printf(TEXT("%.3f ms"), GrandTotalMs));
@@ -307,25 +294,28 @@ void FNiagaraPerformanceGraphTrack::InitTooltip(FTooltipDrawState& InOutTooltip,
 	}
 	else if (HoveredSeries == GpuTimeSeries)
 	{
-		Provider->EnumeratePerformance_RT(EventTime - TimeTolerance, EventTime + TimeTolerance,
+		Provider->EnumeratePerformance_RT(
+			EventTime - TimeTolerance,
+			EventTime + TimeTolerance,
+			false,
 			[&InOutTooltip](const FSystemPerformanceFrame_RT& Frame)
 			{
 				struct FSystemEntry { FStringView Name; double TotalMs; uint64 GpuNumInstances; };
 
 				TArray<FSystemEntry> Entries;
-				double GrandTotalMs = 0.0;
 
-				for (const FSystemPerformanceFrame_RT::FSystemData& Sys : Frame.SystemData)
+				for (const TPair<FString, FSystemPerformanceFrame_RT::FStats>& Pair : Frame.SystemData)
 				{
+					const FSystemPerformanceFrame_RT::FStats& Sys = Pair.Value;
 					if (Sys.GpuNumInstances == 0)
 					{
 						continue;
 					}
-					const double TotalMs = Sys.GpuTotalMicroseconds / 1000.0;
-					Entries.Add({ Sys.SystemName, TotalMs, Sys.GpuNumInstances });
-					GrandTotalMs += TotalMs;
+					Entries.Add({ Pair.Key, Sys.GpuTotalMicroseconds / 1000.0, Sys.GpuNumInstances });
 				}
 
+				const double GrandTotalMs = Frame.AccumulatedStats.GpuTotalMicroseconds / 1000.0;
+
 				Entries.Sort([](const FSystemEntry& A, const FSystemEntry& B) { return A.TotalMs > B.TotalMs; });
 
 				InOutTooltip.AddNameValueTextLine(TEXT("Total:"), FString::Printf(TEXT("%.3f ms"), GrandTotalMs));
@@ -348,20 +338,24 @@ void FNiagaraPerformanceGraphTrack::InitTooltip(FTooltipDrawState& InOutTooltip,
 	}
 	else if (HoveredSeries == MemoryUsageSeries)
 	{
-		Provider->EnumeratePerformance_GT(EventTime - TimeTolerance, EventTime + TimeTolerance,
+		Provider->EnumeratePerformance_GT(
+			EventTime - TimeTolerance,
+			EventTime + TimeTolerance,
+			false,
 			[&InOutTooltip](const FSystemPerformanceFrame_GT& Frame)
 			{
 				struct FSystemEntry { FStringView Name; uint64 MemoryTotalBytes; uint64 NumInstances; };
 
 				TArray<FSystemEntry> Entries;
-				uint64 MemoryTotalBytes = 0;
 
-				for (const FSystemPerformanceFrame_GT::FSystemData& Sys : Frame.SystemData)
+				for (const TPair<FString, FSystemPerformanceFrame_GT::FStats>& Pair : Frame.SystemData)
 				{
-					Entries.Add({ Sys.SystemName, Sys.MemoryBytes, Sys.NumInstances });
-					MemoryTotalBytes += Sys.MemoryBytes;
+					const FSystemPerformanceFrame_GT::FStats& Sys = Pair.Value;
+					Entries.Add({ Pair.Key, Sys.MemoryBytes, Sys.NumInstances });
 				}
 
+				const uint64 MemoryTotalBytes = Frame.AccumulatedStats.MemoryBytes;
+
 				Entries.Sort([](const FSystemEntry& A, const FSystemEntry& B) { return A.MemoryTotalBytes > B.MemoryTotalBytes; });
 
 				InOutTooltip.AddNameValueTextLine(TEXT("Total:"), FString::Printf(TEXT("%.3f KB"), double(MemoryTotalBytes) / 1024.0));
diff --git a/Engine/Plugins/FX/NiagaraInsights/Source/NiagaraInsights/Private/NiagaraProvider.cpp b/Engine/Plugins/FX/NiagaraInsights/Source/NiagaraInsights/Private/NiagaraProvider.cpp
--- a/Engine/Plugins/FX/NiagaraInsights/Source/NiagaraInsights/Private/NiagaraProvider.cpp
+++ b/Engine/Plugins/FX/NiagaraInsights/Source/NiagaraInsights/Private/NiagaraProvider.cpp
@@ -2,6 +2,8 @@
 
 #include "NiagaraProvider.h"
 
+#include "Algo/BinarySearch.h"
+
 namespace UE::NiagaraInsights
 {
 namespace Private
@@ -73,16 +75,25 @@ void FNiagaraProvider::AddSystemPerformance_GT(double EventTime, FStringView Sys
 		PerformanceFrames_GT.AddDefaulted_GetRef().Time = EventTime;
 	}
 
-	FSystemPerformanceFrame_GT::FSystemData& SystemData = PerformanceFrames_GT.Last().SystemData.AddDefaulted_GetRef();
-	SystemData.SystemName = SystemName;
-	SystemData.NumInstances = NumInstances;
-	SystemData.TickGameThreadSeconds = TickGameThreadSeconds;
-	SystemData.TickConcurrentSeconds = TickConcurrentSeconds;
-	SystemData.FinalizeSeconds = FinalizeSeconds;
-	SystemData.EndOfFrameSeconds = EndOfFrameSeconds;
-	SystemData.ActivationSeconds = ActivationSeconds;
-	SystemData.WaitSeconds = WaitSeconds;
-	SystemData.MemoryBytes = MemoryBytes;
+	FSystemPerformanceFrame_GT& Frame_GT = PerformanceFrames_GT.Last();
+	FSystemPerformanceFrame_GT::FStats& Stats_GT = Frame_GT.SystemData.Emplace_GetRef(FString(SystemName), FSystemPerformanceFrame_GT::FStats{}).Value;
+	Stats_GT.NumInstances = NumInstances;
+	Stats_GT.TickGameThreadSeconds = TickGameThreadSeconds;
+	Stats_GT.TickConcurrentSeconds = TickConcurrentSeconds;
+	Stats_GT.FinalizeSeconds = FinalizeSeconds;
+	Stats_GT.EndOfFrameSeconds = EndOfFrameSeconds;
+	Stats_GT.ActivationSeconds = ActivationSeconds;
+	Stats_GT.WaitSeconds = WaitSeconds;
+	Stats_GT.MemoryBytes = MemoryBytes;
+
+	Frame_GT.AccumulatedStats.NumInstances += NumInstances;
+	Frame_GT.AccumulatedStats.TickGameThreadSeconds += TickGameThreadSeconds;
+	Frame_GT.AccumulatedStats.TickConcurrentSeconds += TickConcurrentSeconds;
+	Frame_GT.AccumulatedStats.FinalizeSeconds += FinalizeSeconds;
+	Frame_GT.AccumulatedStats.EndOfFrameSeconds += EndOfFrameSeconds;
+	Frame_GT.AccumulatedStats.ActivationSeconds += ActivationSeconds;
+	Frame_GT.AccumulatedStats.WaitSeconds += WaitSeconds;
+	Frame_GT.AccumulatedStats.MemoryBytes += MemoryBytes;
 }
 
 void FNiagaraProvider::AddSystemPerformance_RT(double EventTime, FStringView SystemName, uint64 NumInstances, double RenderUpdateSeconds, double GetDynamicMeshElementsSeconds, uint64 GpuNumInstances, uint64 GpuTotalMicroseconds)
@@ -95,48 +106,68 @@ void FNiagaraProvider::AddSystemPerformance_RT(double EventTime, FStringView Sys
 		PerformanceFrames_RT.AddDefaulted_GetRef().Time = EventTime;
 	}
 
-	FSystemPerformanceFrame_RT::FSystemData& SystemData = PerformanceFrames_RT.Last().SystemData.AddDefaulted_GetRef();
-	SystemData.SystemName = SystemName;
-	SystemData.NumInstances = NumInstances;
-	SystemData.RenderUpdateSeconds = RenderUpdateSeconds;
-	SystemData.GetDynamicMeshElementsSeconds = GetDynamicMeshElementsSeconds;
-	SystemData.GpuNumInstances = GpuNumInstances;
-	SystemData.GpuTotalMicroseconds = GpuTotalMicroseconds;
+	FSystemPerformanceFrame_RT& Frame_RT = PerformanceFrames_RT.Last();
+	FSystemPerformanceFrame_RT::FStats& Stats_RT = Frame_RT.SystemData.Emplace_GetRef(FString(SystemName), FSystemPerformanceFrame_RT::FStats{}).Value;
+	Stats_RT.NumInstances = NumInstances;
+	Stats_RT.RenderUpdateSeconds = RenderUpdateSeconds;
+	Stats_RT.GetDynamicMeshElementsSeconds = GetDynamicMeshElementsSeconds;
+	Stats_RT.GpuNumInstances = GpuNumInstances;
+	Stats_RT.GpuTotalMicroseconds = GpuTotalMicroseconds;
+
+	Frame_RT.AccumulatedStats.NumInstances += NumInstances;
+	Frame_RT.AccumulatedStats.RenderUpdateSeconds += RenderUpdateSeconds;
+	Frame_RT.AccumulatedStats.GetDynamicMeshElementsSeconds += GetDynamicMeshElementsSeconds;
+	Frame_RT.AccumulatedStats.GpuNumInstances += GpuNumInstances;
+	Frame_RT.AccumulatedStats.GpuTotalMicroseconds += GpuTotalMicroseconds;
 }
 
-void FNiagaraProvider::EnumeratePerformance_GT(double StartTime, double EndTime, TFunction<void(const FSystemPerformanceFrame_GT&)> Evaluate) const
+void FNiagaraProvider::EnumeratePerformance_GT(double StartTime, double EndTime, bool bIncludeOverlapping, TFunction<void(const FSystemPerformanceFrame_GT&)> Evaluate) const
 {
 	TraceServices::FProviderReadScopeLock ReadLock(*this);
 
-	for (const FSystemPerformanceFrame_GT& Frame : PerformanceFrames_GT)
+	const int32 Num = PerformanceFrames_GT.Num();
+	if (Num == 0)
 	{
-		if (Frame.Time > EndTime)
-		{
-			return;
-		}
-		if (Frame.Time < StartTime)
-		{
-			continue;
-		}
-		Evaluate(Frame);
+		return;
+	}
+
+	int32 StartIndex = Algo::LowerBoundBy(PerformanceFrames_GT, StartTime, &FSystemPerformanceFrame_GT::Time);
+	int32 EndIndex = Algo::UpperBoundBy(PerformanceFrames_GT, EndTime, &FSystemPerformanceFrame_GT::Time);
+
+	if (bIncludeOverlapping)
+	{
+		StartIndex = FMath::Max(StartIndex - 1, 0);
+		EndIndex = FMath::Min(EndIndex + 1, Num);
+	}
+
+	for (int32 i = StartIndex; i < EndIndex; ++i)
+	{
+		Evaluate(PerformanceFrames_GT[i]);
 	}
 }
 
-void FNiagaraProvider::EnumeratePerformance_RT(double StartTime, double EndTime, TFunction<void(const FSystemPerformanceFrame_RT&)> Evaluate) const
+void FNiagaraProvider::EnumeratePerformance_RT(double StartTime, double EndTime, bool bIncludeOverlapping, TFunction<void(const FSystemPerformanceFrame_RT&)> Evaluate) const
 {
 	TraceServices::FProviderReadScopeLock ReadLock(*this);
 
-	for (const FSystemPerformanceFrame_RT& Frame : PerformanceFrames_RT)
+	const int32 Num = PerformanceFrames_RT.Num();
+	if (Num == 0)
 	{
-		if (Frame.Time > EndTime)
-		{
-			return;
-		}
-		if (Frame.Time < StartTime)
-		{
-			continue;
-		}
-		Evaluate(Frame);
+		return;
+	}
+
+	int32 StartIndex = Algo::LowerBoundBy(PerformanceFrames_RT, StartTime, &FSystemPerformanceFrame_RT::Time);
+	int32 EndIndex = Algo::UpperBoundBy(PerformanceFrames_RT, EndTime, &FSystemPerformanceFrame_RT::Time);
+
+	if (bIncludeOverlapping)
+	{
+		StartIndex = FMath::Max(StartIndex - 1, 0);
+		EndIndex = FMath::Min(EndIndex + 1, Num);
+	}
+
+	for (int32 i = StartIndex; i < EndIndex; ++i)
+	{
+		Evaluate(PerformanceFrames_RT[i]);
 	}
 }
 
@@ -174,15 +205,49 @@ void FNiagaraProvider::EnumerateComponentEvent(double StartTime, double EndTime,
 {
 	TraceServices::FProviderReadScopeLock ReadLock(*this);
 
-	for (const FLifetimeEvent& Event : LifetimeEvents)
+	const int32 StartIndex = Algo::LowerBoundBy(LifetimeEvents, StartTime, &FLifetimeEvent::Time);
+	const int32 Num = LifetimeEvents.Num();
+	for (int32 i = StartIndex; i < Num; ++i)
 	{
+		const FLifetimeEvent& Event = LifetimeEvents[i];
 		if (Event.Time > EndTime)
 		{
-			return;
+			break;
 		}
-		if (Event.Time < StartTime)
+		Evaluate(Event);
+	}
+}
+
+void FNiagaraProvider::AddDataChannelPublish(double EventTime, FStringView SourceName, bool bGpuRequest, bool bVisibleToGame, bool bVisibleToCPUSims, bool bVisibleToGPUSims, uint32 NumInstances, uint32 NumInstanceAllocated)
+{
+	TraceServices::FProviderEditScopeLock WriteLock(*this);
+
+	FDataChannelEvent& Event = DataChannelEvents.AddDefaulted_GetRef();
+	Event.Time = EventTime;
+	Event.Payload.Emplace<FDataChannelEvent::FPublish>(FString(SourceName), bGpuRequest, bVisibleToGame, bVisibleToCPUSims, bVisibleToGPUSims, NumInstances, NumInstanceAllocated);
+}
+
+void FNiagaraProvider::AddDataChannelWrite(double EventTime, FStringView DataChannelName, FStringView SourceName, int32 NumInstances, bool bVisibleToGame, bool bVisibleToCPU, bool bVisibleToGPU)
+{
+	TraceServices::FProviderEditScopeLock WriteLock(*this);
+
+	FDataChannelEvent& Event = DataChannelEvents.AddDefaulted_GetRef();
+	Event.Time = EventTime;
+	Event.Payload.Emplace<FDataChannelEvent::FWrite>(FString(DataChannelName), FString(SourceName), NumInstances, bVisibleToGame, bVisibleToCPU, bVisibleToGPU);
+}
+
+void FNiagaraProvider::EnumerateDataChannelEvent(double StartTime, double EndTime, TFunction<void(const FDataChannelEvent&)> Evaluate) const
+{
+	TraceServices::FProviderReadScopeLock ReadLock(*this);
+
+	const int32 StartIndex = Algo::LowerBoundBy(DataChannelEvents, StartTime, &FDataChannelEvent::Time);
+	const int32 Num = DataChannelEvents.Num();
+	for (int32 i = StartIndex; i < Num; ++i)
+	{
+		const FDataChannelEvent& Event = DataChannelEvents[i];
+		if (Event.Time > EndTime)
 		{
-			continue;
+			break;
 		}
 		Evaluate(Event);
 	}
diff --git a/Engine/Plugins/FX/NiagaraInsights/Source/NiagaraInsights/Private/NiagaraProvider.h b/Engine/Plugins/FX/NiagaraInsights/Source/NiagaraInsights/Private/NiagaraProvider.h
--- a/Engine/Plugins/FX/NiagaraInsights/Source/NiagaraInsights/Private/NiagaraProvider.h
+++ b/Engine/Plugins/FX/NiagaraInsights/Source/NiagaraInsights/Private/NiagaraProvider.h
@@ -14,9 +14,8 @@ namespace UE::NiagaraInsights
 {
 struct FSystemPerformanceFrame_GT
 {
-	struct FSystemData
+	struct FStats
 	{
-		FString SystemName;
 		uint64 NumInstances = 0;
 		double TickGameThreadSeconds = 0.0;
 		double TickConcurrentSeconds = 0.0;
@@ -24,18 +23,18 @@ struct FSystemPerformanceFrame_GT
 		double EndOfFrameSeconds = 0.0;
 		double ActivationSeconds = 0.0;
 		double WaitSeconds = 0.0;
-		uint64 MemoryBytes = 0.0;
+		uint64 MemoryBytes = 0;
 	};
 
 	double Time = 0.0;
-	TArray<FSystemData> SystemData;
+	TArray<TPair<FString, FStats>> SystemData;
+	FStats AccumulatedStats;
 };
 
 struct FSystemPerformanceFrame_RT
 {
-	struct FSystemData
+	struct FStats
 	{
-		FString SystemName;
 		uint64 NumInstances = 0;
 		double RenderUpdateSeconds = 0.0;
 		double GetDynamicMeshElementsSeconds = 0.0;
@@ -45,7 +44,8 @@ struct FSystemPerformanceFrame_RT
 	};
 
 	double Time = 0.0;
-	TArray<FSystemData> SystemData;
+	TArray<TPair<FString, FStats>> SystemData;
+	FStats AccumulatedStats;
 };
 
 struct FLifetimeEvent
@@ -72,6 +72,33 @@ struct FLifetimeEvent
 	TVariant<FActivate, FDeactivate, FComplete> Payload;
 };
 
+struct FDataChannelEvent
+{
+	struct FPublish
+	{
+		FString SourceName;
+		bool bGpuRequest = false;
+		bool bVisibleToGame = false;
+		bool bVisibleToCPUSims = false;
+		bool bVisibleToGPUSims = false;
+		uint32 NumInstances = 0;
+		uint32 NumInstanceAllocated = 0;
+	};
+
+	struct FWrite
+	{
+		FString DataChannelName;
+		FString SourceName;
+		int32	NumInstances = 0;
+		bool	bVisibleToGame = false;
+		bool	bVisibleToCPU = false;
+		bool	bVisibleToGPU = false;
+	};
+
+	double Time = 0.0;
+	TVariant<FPublish, FWrite> Payload;
+};
+
 class FNiagaraProvider : public TraceServices::IProvider, public TraceServices::IEditableProvider
 {
 public:
@@ -97,16 +124,28 @@ public:
 	void AddSystemPerformance_RT(double EventTime, FStringView SystemName, uint64 NumInstances, double RenderUpdateSeconds, double GetDynamicMeshElementsSeconds, uint64 GpuNumInstances, uint64 GpuTotalMicroseconds);
 
 	// Enumerate a range of game thread frames
-	void EnumeratePerformance_GT(double StartTime, double EndTime, TFunction<void(const FSystemPerformanceFrame_GT&)> Evaluate) const;
+	void EnumeratePerformance_GT(double StartTime, double EndTime, bool bIncludeOverlapping, TFunction<void(const FSystemPerformanceFrame_GT&)> Evaluate) const;
 	// Enumerate a range of render thread frames
-	void EnumeratePerformance_RT(double StartTime, double EndTime, TFunction<void(const FSystemPerformanceFrame_RT&)> Evaluate) const;
+	void EnumeratePerformance_RT(double StartTime, double EndTime, bool bIncludeOverlapping, TFunction<void(const FSystemPerformanceFrame_RT&)> Evaluate) const;
 
+	// Add component activate event
 	void AddComponentActivate(double EventTime, FStringView ComponentName, bool bReset, bool bIsScalabilityCull, bool bAwaitingActivationDueToNotReady);
+	// Add component deactivate event
 	void AddComponentDeactivate(double EventTime, FStringView ComponentName, bool bImmediate, bool bIsScalabilityCull);
+	// Add component complete event
 	void AddComponentComplete(double EventTime, FStringView ComponentName, bool bExternalCompletion);
 
+	// Enumerate a range of component events
 	void EnumerateComponentEvent(double StartTime, double EndTime, TFunction<void(const FLifetimeEvent&)> Evaluate) const;
 
+	// Add data channel publish event
+	void AddDataChannelPublish(double EventTime, FStringView SourceName, bool bGpuRequest, bool bVisibleToGame, bool bVisibleToCPUSims, bool bVisibleToGPUSims, uint32 NumInstances, uint32 NumInstanceAllocated);
+	// Add data channel write event
+	void AddDataChannelWrite(double EventTime, FStringView DataChannelName, FStringView SourceName, int32 NumInstances, bool bVisibleToGame, bool bVisibleToCPU, bool bVisibleToGPU);
+
+	// Enumerate a range of data channel events
+	void EnumerateDataChannelEvent(double StartTime, double EndTime, TFunction<void(const FDataChannelEvent&)> Evaluate) const;
+
 private:
 	mutable TraceServices::FProviderLock	ProviderLock;
 
@@ -116,6 +155,8 @@ private:
 	TArray<FSystemPerformanceFrame_RT>		PerformanceFrames_RT;
 
 	TArray<FLifetimeEvent>					LifetimeEvents;
+
+	TArray<FDataChannelEvent>				DataChannelEvents;
 };
 
 } //namespace UE::NiagaraInsights
diff --git a/Engine/Plugins/FX/NiagaraInsights/Source/NiagaraInsights/Private/NiagaraTimingViewSession.cpp b/Engine/Plugins/FX/NiagaraInsights/Source/NiagaraInsights/Private/NiagaraTimingViewSession.cpp
--- a/Engine/Plugins/FX/NiagaraInsights/Source/NiagaraInsights/Private/NiagaraTimingViewSession.cpp
+++ b/Engine/Plugins/FX/NiagaraInsights/Source/NiagaraInsights/Private/NiagaraTimingViewSession.cpp
@@ -1,6 +1,7 @@
 // Copyright Epic Games, Inc. All Rights Reserved.
 
 #include "NiagaraTimingViewSession.h"
+#include "NiagaraDataChannelTrack.h"
 #include "NiagaraInstanceLifecycleTrack.h"
 #include "NiagaraPerformanceGraphTrack.h"
 #include "NiagaraProvider.h"
@@ -17,6 +18,7 @@ void FNiagaraTimingViewSession::OnBeginSession(UE::Insights::Timing::ITimingView
 {
 	TimingViewSession = &InTimingViewSession;
 
+	DataChannelTrack.Reset();
 	InstanceLifecycleTrack.Reset();
 	PerformanceGraphTrack.Reset();
 
@@ -27,6 +29,7 @@ void FNiagaraTimingViewSession::OnEndSession(UE::Insights::Timing::ITimingViewSe
 {
 	InTimingViewSession.OnSelectionChanged().RemoveAll(this);
 
+	DataChannelTrack.Reset();
 	InstanceLifecycleTrack.Reset();
 	PerformanceGraphTrack.Reset();
 	TimingViewSession = nullptr;
@@ -55,16 +58,23 @@ void FNiagaraTimingViewSession::Tick(UE::Insights::Timing::ITimingViewSession& I
 
 		if (!InstanceLifecycleTrack.IsValid())
 		{
+			constexpr int32 TrackOrder = MAX_int32 / 2;
+		
 			PerformanceGraphTrack = MakeShared<FNiagaraPerformanceGraphTrack>(*this);
 			PerformanceGraphTrack->SetVisibilityFlag(bTrackVisible);
-			PerformanceGraphTrack->SetOrder((MAX_int32 / 2));
+			PerformanceGraphTrack->SetOrder(TrackOrder);
 			InTimingViewSession.AddScrollableTrack(PerformanceGraphTrack);
 
 			InstanceLifecycleTrack = MakeShared<FNiagaraInstanceLifecycleTrack>(*this);
 			InstanceLifecycleTrack->SetVisibilityFlag(bTrackVisible);
-			InstanceLifecycleTrack->SetOrder((MAX_int32 / 2) + 1);
+			InstanceLifecycleTrack->SetOrder(TrackOrder + 1);
 			InTimingViewSession.AddScrollableTrack(InstanceLifecycleTrack);
 
+			DataChannelTrack = MakeShared<FNiagaraDataChannelTrack>(*this);
+			DataChannelTrack->SetVisibilityFlag(bTrackVisible);
+			DataChannelTrack->SetOrder(TrackOrder + 2);
+			InTimingViewSession.AddScrollableTrack(DataChannelTrack);
+
 			InTimingViewSession.InvalidateScrollableTracksOrder();
 		}
 	}
@@ -94,6 +104,10 @@ void FNiagaraTimingViewSession::ToggleNiagaraTrack()
 {
 	bTrackVisible = !bTrackVisible;
 
+	if (DataChannelTrack)
+	{
+		DataChannelTrack->SetVisibilityFlag(bTrackVisible);
+	}
 	if (InstanceLifecycleTrack)
 	{
 		InstanceLifecycleTrack->SetVisibilityFlag(bTrackVisible);
diff --git a/Engine/Plugins/FX/NiagaraInsights/Source/NiagaraInsights/Private/NiagaraTimingViewSession.h b/Engine/Plugins/FX/NiagaraInsights/Source/NiagaraInsights/Private/NiagaraTimingViewSession.h
--- a/Engine/Plugins/FX/NiagaraInsights/Source/NiagaraInsights/Private/NiagaraTimingViewSession.h
+++ b/Engine/Plugins/FX/NiagaraInsights/Source/NiagaraInsights/Private/NiagaraTimingViewSession.h
@@ -13,6 +13,7 @@ class FMenuBuilder;
 namespace UE::NiagaraInsights
 {
 
+class FNiagaraDataChannelTrack;
 class FNiagaraInstanceLifecycleTrack;
 class FNiagaraPerformanceGraphTrack;
 
@@ -45,6 +46,7 @@ private:
 	const TraceServices::IAnalysisSession*		AnalysisSession = nullptr;
 	UE::Insights::Timing::ITimingViewSession*	TimingViewSession = nullptr;
 
+	TSharedPtr<FNiagaraDataChannelTrack>		DataChannelTrack;
 	TSharedPtr<FNiagaraInstanceLifecycleTrack>	InstanceLifecycleTrack;
 	TSharedPtr<FNiagaraPerformanceGraphTrack>	PerformanceGraphTrack;
 
diff --git a/Engine/Plugins/FX/NiagaraInsights/Source/NiagaraInsights/Private/Widgets/SNiagaraRangeStatsView.cpp b/Engine/Plugins/FX/NiagaraInsights/Source/NiagaraInsights/Private/Widgets/SNiagaraRangeStatsView.cpp
--- a/Engine/Plugins/FX/NiagaraInsights/Source/NiagaraInsights/Private/Widgets/SNiagaraRangeStatsView.cpp
+++ b/Engine/Plugins/FX/NiagaraInsights/Source/NiagaraInsights/Private/Widgets/SNiagaraRangeStatsView.cpp
@@ -348,14 +348,18 @@ void SNiagaraRangeStatsView::RebuildStats(double StartTime, double EndTime)
 	TMap<FString, FAggregated> AggMap;
 
 	int32 GTFrameCount = 0;
-	Provider->EnumeratePerformance_GT(StartTime, EndTime,
+	Provider->EnumeratePerformance_GT(
+		StartTime,
+		EndTime,
+		false,
 		[&](const FSystemPerformanceFrame_GT& Frame)
 		{
 			++GTFrameCount;
-			for (const FSystemPerformanceFrame_GT::FSystemData& Sys : Frame.SystemData)
+			for (const TPair<FString, FSystemPerformanceFrame_GT::FStats>& Pair : Frame.SystemData)
 			{
-				FAggregated& A = AggMap.FindOrAdd(Sys.SystemName);
-				A.SystemName = Sys.SystemName;
+				const FSystemPerformanceFrame_GT::FStats& Sys = Pair.Value;
+				FAggregated& A = AggMap.FindOrAdd(Pair.Key);
+				A.SystemName = Pair.Key;
 
 				A.GTSamples++;
 				A.GTTotalInstances += (double)Sys.NumInstances;
@@ -369,13 +373,17 @@ void SNiagaraRangeStatsView::RebuildStats(double StartTime, double EndTime)
 			}
 		});
 
-	Provider->EnumeratePerformance_RT(StartTime, EndTime,
+	Provider->EnumeratePerformance_RT(
+		StartTime,
+		EndTime,
+		false,
 		[&](const FSystemPerformanceFrame_RT& Frame)
 		{
-			for (const FSystemPerformanceFrame_RT::FSystemData& Sys : Frame.SystemData)
+			for (const TPair<FString, FSystemPerformanceFrame_RT::FStats>& Pair : Frame.SystemData)
 			{
-				FAggregated& A = AggMap.FindOrAdd(Sys.SystemName);
-				A.SystemName = Sys.SystemName;
+				const FSystemPerformanceFrame_RT::FStats& Sys = Pair.Value;
+				FAggregated& A = AggMap.FindOrAdd(Pair.Key);
+				A.SystemName = Pair.Key;
 
 				A.RTSamples++;
 				A.RTTotalInstances += (double)Sys.NumInstances;
