- Don't show the Niagara tracks if we have no data

#rnx

#Synced-CL 52062699

diff --git a/Engine/Plugins/FX/NiagaraInsights/Source/NiagaraInsights/Private/NiagaraProvider.cpp b/Engine/Plugins/FX/NiagaraInsights/Source/NiagaraInsights/Private/NiagaraProvider.cpp
--- a/Engine/Plugins/FX/NiagaraInsights/Source/NiagaraInsights/Private/NiagaraProvider.cpp
+++ b/Engine/Plugins/FX/NiagaraInsights/Source/NiagaraInsights/Private/NiagaraProvider.cpp
@@ -65,6 +65,15 @@ FName FNiagaraProvider::GetProviderName()
 	return Private::ProviderName;
 }
 
+bool FNiagaraProvider::HasAnyData() const
+{
+	TraceServices::FProviderReadScopeLock ReadLock(*this);
+	return PerformanceFrames_GT.Num() > 0
+		|| PerformanceFrames_RT.Num() > 0
+		|| LifetimeEvents.Num() > 0
+		|| DataChannelEvents.Num() > 0;
+}
+
 void FNiagaraProvider::AddSystemPerformance_GT(double EventTime, FStringView SystemName, uint64 NumInstances, double TickGameThreadSeconds, double TickConcurrentSeconds, double FinalizeSeconds, double EndOfFrameSeconds, double ActivationSeconds, double WaitSeconds, uint64 MemoryBytes)
 {
 	TraceServices::FProviderEditScopeLock WriteLock(*this);
diff --git a/Engine/Plugins/FX/NiagaraInsights/Source/NiagaraInsights/Private/NiagaraProvider.h b/Engine/Plugins/FX/NiagaraInsights/Source/NiagaraInsights/Private/NiagaraProvider.h
--- a/Engine/Plugins/FX/NiagaraInsights/Source/NiagaraInsights/Private/NiagaraProvider.h
+++ b/Engine/Plugins/FX/NiagaraInsights/Source/NiagaraInsights/Private/NiagaraProvider.h
@@ -118,6 +118,9 @@ public:
 
 	static FName GetProviderName();
 
+	// Returns true if any trace events have been received
+	bool HasAnyData() const;
+
 	// Add system performance data for the game thread
 	void AddSystemPerformance_GT(double EventTime, FStringView SystemName, uint64 NumInstances, double TickGameThreadSeconds, double TickConcurrentSeconds, double FinalizeSeconds, double EndOfFrameSeconds, double ActivationSeconds, double WaitSeconds, uint64 MemoryBytes);
 	// Add system performnace data for the render thread
diff --git a/Engine/Plugins/FX/NiagaraInsights/Source/NiagaraInsights/Private/NiagaraTimingViewSession.cpp b/Engine/Plugins/FX/NiagaraInsights/Source/NiagaraInsights/Private/NiagaraTimingViewSession.cpp
--- a/Engine/Plugins/FX/NiagaraInsights/Source/NiagaraInsights/Private/NiagaraTimingViewSession.cpp
+++ b/Engine/Plugins/FX/NiagaraInsights/Source/NiagaraInsights/Private/NiagaraTimingViewSession.cpp
@@ -53,10 +53,7 @@ void FNiagaraTimingViewSession::Tick(UE::Insights::Timing::ITimingViewSession& I
 	{
 		TraceServices::FAnalysisSessionReadScope SessionReadScope(*AnalysisSession);
 
-		//const FNiagaraTimelineDataPtr& TimelineData = NiagaraProvider->GetTimelineData();
-		//-TODO: Check we have data
-
-		if (!InstanceLifecycleTrack.IsValid())
+		if (!InstanceLifecycleTrack.IsValid() && NiagaraProvider->HasAnyData())
 		{
 			constexpr int32 TrackOrder = MAX_int32 / 2;
 		
